Fix version selector compile errors on Mac.

diff --git a/Engine/Source/Programs/UnrealVersionSelector/Private/FileAssociation.cpp b/Engine/Source/Programs/UnrealVersionSelector/Private/FileAssociation.cpp
--- a/Engine/Source/Programs/UnrealVersionSelector/Private/FileAssociation.cpp
+++ b/Engine/Source/Programs/UnrealVersionSelector/Private/FileAssociation.cpp
@@ -9,14 +9,14 @@ TSharedPtr<FJsonObject> LoadJson(const FString &FileName)
 
 	if (!FFileHelper::LoadFileToString(FileContents, *FileName))
 	{
-		return false;
+		return TSharedPtr<FJsonObject>(NULL);
 	}
 
 	TSharedPtr< FJsonObject > JsonObject;
 	TSharedRef< TJsonReader<> > Reader = TJsonReaderFactory<>::Create(FileContents);
 	if (!FJsonSerializer::Deserialize(Reader, JsonObject) || !JsonObject.IsValid())
 	{
-		return false;
+		return TSharedPtr<FJsonObject>(NULL);
 	}
 
 	return JsonObject;
