Fix UnrealVersionSelector compile error on Mac.

diff --git a/Engine/Source/Programs/UnrealVersionSelector/Private/UnrealVersionSelector.cpp b/Engine/Source/Programs/UnrealVersionSelector/Private/UnrealVersionSelector.cpp
--- a/Engine/Source/Programs/UnrealVersionSelector/Private/UnrealVersionSelector.cpp
+++ b/Engine/Source/Programs/UnrealVersionSelector/Private/UnrealVersionSelector.cpp
@@ -312,4 +312,19 @@ int Main(const TArray<FString> &Arguments)
 
 	#include "HideWindowsPlatformTypes.h"
 
+#else
+
+	int main(int ArgC, const char *ArgV[])
+	{
+		FCommandLine::Set(TEXT(""));
+		
+		TArray<FString> Arguments;
+		for (int Idx = 1; Idx < ArgC; Idx++)
+		{
+			Arguments.Add(ArgV[Idx]);
+		}
+		
+		return Main(Arguments);
+	}
+
 #endif
\ No newline at end of file
