Allow engine to query installed versions, so we can handle .uproject engine associations from inside the editor and at runtime.

diff --git a/Engine/Source/Programs/UnrealVersionSelector/Private/FileAssociation.cpp b/Engine/Source/Programs/UnrealVersionSelector/Private/FileAssociation.cpp
--- a/Engine/Source/Programs/UnrealVersionSelector/Private/FileAssociation.cpp
+++ b/Engine/Source/Programs/UnrealVersionSelector/Private/FileAssociation.cpp
@@ -71,7 +71,7 @@ bool GetEngineIdForProject(const FString &ProjectFileName, FString &OutId)
 bool GetEngineRootDirForProject(const FString &ProjectFileName, FString &OutRootDir)
 {
 	FString Id;
-	if (GetEngineIdForProject(ProjectFileName, Id) && GetEngineRootDirFromId(Id, OutRootDir))
+	if (GetEngineIdForProject(ProjectFileName, Id) && FPlatformMisc::GetEngineRootDirFromIdentifier(Id, OutRootDir))
 	{
 		return true;
 	}
diff --git a/Engine/Source/Programs/UnrealVersionSelector/Private/GenericPlatform/GenericPlatformInstallation.cpp b/Engine/Source/Programs/UnrealVersionSelector/Private/GenericPlatform/GenericPlatformInstallation.cpp
--- a/Engine/Source/Programs/UnrealVersionSelector/Private/GenericPlatform/GenericPlatformInstallation.cpp
+++ b/Engine/Source/Programs/UnrealVersionSelector/Private/GenericPlatform/GenericPlatformInstallation.cpp
@@ -7,10 +7,11 @@
 
 bool FGenericPlatformInstallation::GetLauncherVersionSelector(FString &OutFileName)
 {
-	FConfigFile *ConfigFile = GetConfigFile();
+	FConfigFile ConfigFile;
+	ConfigFile.Read(GetConfigFilePath());
 
 	FString InstallationPath;
-	if (!ConfigFile->GetString(TEXT("Launcher"), TEXT("InstalledPath"), InstallationPath))
+	if (!ConfigFile.GetString(TEXT("Launcher"), TEXT("InstalledPath"), InstallationPath))
 	{
 		return false;
 	}
@@ -26,31 +27,27 @@ bool FGenericPlatformInstallation::GetLauncherVersionSelector(FString &OutFileNa
 
 bool FGenericPlatformInstallation::RegisterEngineInstallation(const FString &Id, const FString &RootDirName)
 {
-	FConfigFile *ConfigFile = GetConfigFile();
-	FConfigSection &Section = ConfigFile->FindOrAdd(TEXT("Installations"));
-	Section.AddUnique(*Id, RootDirName);
-	ConfigFile->Dirty = true;
-	return ConfigFile->Write(GetConfigFilePath());
-}
+	FString NormalizedRootDirName = RootDirName;
+	FPaths::NormalizeDirectoryName(NormalizedRootDirName);
 
-bool FGenericPlatformInstallation::UnregisterEngineInstallation(const FString &Id)
-{
-	FConfigFile *ConfigFile = GetConfigFile();
-	FConfigSection &Section = ConfigFile->FindOrAdd(TEXT("Installations"));
-	Section.Remove(*Id);
-	ConfigFile->Dirty = true;
-	return ConfigFile->Write(GetConfigFilePath());
+	FConfigFile ConfigFile;
+	ConfigFile.Read(GetConfigFilePath());
+
+	FConfigSection &Section = ConfigFile.FindOrAdd(TEXT("Installations"));
+	Section.AddUnique(*Id, NormalizedRootDirName);
+	ConfigFile.Dirty = true;
+	return ConfigFile.Write(GetConfigFilePath());
 }
 
-void FGenericPlatformInstallation::EnumerateEngineInstallations(TMap<FString, FString> &OutInstallations)
+bool FGenericPlatformInstallation::UnregisterEngineInstallation(const FString &Id)
 {
-	FConfigFile *ConfigFile = GetConfigFile();
-	FConfigSection &Section = ConfigFile->FindOrAdd(TEXT("Installations"));
+	FConfigFile ConfigFile;
+	ConfigFile.Read(GetConfigFilePath());
 
-	for (TMultiMap<FName, FString>::TConstIterator Iter(Section); Iter; ++Iter)
-	{
-		OutInstallations.Add(Iter.Key().ToString(), Iter.Value());
-	}
+	FConfigSection &Section = ConfigFile.FindOrAdd(TEXT("Installations"));
+	Section.Remove(*Id);
+	ConfigFile.Dirty = true;
+	return ConfigFile.Write(GetConfigFilePath());
 }
 
 bool FGenericPlatformInstallation::IsSourceDistribution(const FString &EngineRootDir)
@@ -135,15 +132,3 @@ FString FGenericPlatformInstallation::GetConfigFilePath()
 	return FString(FPlatformProcess::ApplicationSettingsDir()) / FString(EPIC_PRODUCT_IDENTIFIER) / FString(TEXT("Install.ini"));
 }
 
-FConfigFile *FGenericPlatformInstallation::GetConfigFile()
-{
-	static FConfigFile *ConfigFile = NULL;
-	if (ConfigFile == NULL)
-	{
-		ConfigFile = new FConfigFile();
-		ConfigFile->Read(GetConfigFilePath());
-	}
-	return ConfigFile;
-}
-
-
diff --git a/Engine/Source/Programs/UnrealVersionSelector/Private/GenericPlatform/GenericPlatformInstallation.h b/Engine/Source/Programs/UnrealVersionSelector/Private/GenericPlatform/GenericPlatformInstallation.h
--- a/Engine/Source/Programs/UnrealVersionSelector/Private/GenericPlatform/GenericPlatformInstallation.h
+++ b/Engine/Source/Programs/UnrealVersionSelector/Private/GenericPlatform/GenericPlatformInstallation.h
@@ -9,10 +9,11 @@ struct FGenericPlatformInstallation
 	/// Get the path to the installed build of the version selector
 	static bool GetLauncherVersionSelector(FString &OutFileName);
 
-	// Engine installations
+	// Register an engine installation
 	static bool RegisterEngineInstallation(const FString &Id, const FString &RootDirName);
+
+	// Unregister an engine installation
 	static bool UnregisterEngineInstallation(const FString &Id);
-	static void EnumerateEngineInstallations(TMap<FString, FString> &OutInstallations);
 
 	// Test whether a given engine root directory is a source build
 	static bool IsSourceDistribution(const FString &EngineRootDir);
diff --git a/Engine/Source/Programs/UnrealVersionSelector/Private/PlatformInstallation.cpp b/Engine/Source/Programs/UnrealVersionSelector/Private/PlatformInstallation.cpp
--- a/Engine/Source/Programs/UnrealVersionSelector/Private/PlatformInstallation.cpp
+++ b/Engine/Source/Programs/UnrealVersionSelector/Private/PlatformInstallation.cpp
@@ -22,23 +22,23 @@ int32 ParseReleaseVersion(const FString &Version)
 
 bool GetDefaultEngineId(FString &OutId)
 {
-	TMap<FString, FString> Installations;
-	FPlatformInstallation::EnumerateEngineInstallations(Installations);
+	TArray<TPair<FString, FString>> Installations;
+	FPlatformMisc::EnumerateEngineInstallations(Installations);
 
 	bool bRes = false;
 	if (Installations.Num() > 0)
 	{
 		// Default to the first install
-		OutId = TMap<FString, FString>::TConstIterator(Installations).Key();
+		OutId = Installations[0].Key;
 
 		// Try to find the highest versioned install
 		int32 BestVersion = INDEX_NONE;
-		for (TMap<FString, FString>::TConstIterator Iter(Installations); Iter; ++Iter)
+		for (TArray<TPair<FString, FString>>::TConstIterator Iter(Installations); Iter; ++Iter)
 		{
-			int32 Version = ParseReleaseVersion(Iter.Key());
+			int32 Version = ParseReleaseVersion(Iter->Key);
 			if (Version > BestVersion)
 			{
-				OutId = Iter.Key();
+				OutId = Iter->Key;
 				BestVersion = Version;
 			}
 		}
@@ -49,38 +49,5 @@ bool GetDefaultEngineId(FString &OutId)
 bool GetDefaultEngineRootDir(FString &OutDirName)
 {
 	FString Id;
-	return GetDefaultEngineId(Id) && GetEngineRootDirFromId(Id, OutDirName);
+	return GetDefaultEngineId(Id) && FPlatformMisc::GetEngineRootDirFromIdentifier(Id, OutDirName);
 }
-
-bool GetEngineRootDirFromId(const FString &Id, FString &OutDirName)
-{
-	TMap<FString, FString> Installations;
-	FPlatformInstallation::EnumerateEngineInstallations(Installations);
-
-	const FString *DirName = Installations.Find(Id);
-	if (DirName != NULL)
-	{
-		OutDirName = *DirName;
-		return true;
-	}
-
-	return false;
-}
-
-bool GetEngineIdFromRootDir(const FString &DirName, FString &OutId)
-{
-	TMap<FString, FString> Installations;
-	FPlatformInstallation::EnumerateEngineInstallations(Installations);
-
-	// Find the label for the given directory
-	for (TMap<FString, FString>::TConstIterator Iter(Installations); Iter; ++Iter)
-	{
-		if (Iter.Value() == DirName)
-		{
-			OutId = Iter.Key();
-			return true;
-		}
-	}
-	return false;
-}
-
diff --git a/Engine/Source/Programs/UnrealVersionSelector/Private/PlatformInstallation.h b/Engine/Source/Programs/UnrealVersionSelector/Private/PlatformInstallation.h
--- a/Engine/Source/Programs/UnrealVersionSelector/Private/PlatformInstallation.h
+++ b/Engine/Source/Programs/UnrealVersionSelector/Private/PlatformInstallation.h
@@ -18,5 +18,3 @@ int32 ParseReleaseVersion(const FString &Version);
 bool GetDefaultEngineId(FString &OutId);
 bool GetDefaultEngineRootDir(FString &OutId);
 
-bool GetEngineRootDirFromId(const FString &Id, FString &OutRootDir);
-bool GetEngineIdFromRootDir(const FString &RootDir, FString &OutId);
diff --git a/Engine/Source/Programs/UnrealVersionSelector/Private/UnrealVersionSelector.cpp b/Engine/Source/Programs/UnrealVersionSelector/Private/UnrealVersionSelector.cpp
--- a/Engine/Source/Programs/UnrealVersionSelector/Private/UnrealVersionSelector.cpp
+++ b/Engine/Source/Programs/UnrealVersionSelector/Private/UnrealVersionSelector.cpp
@@ -25,22 +25,30 @@ bool RegisterCurrentEngineDirectory()
 
 	// Launch the installed version selector as elevated to register this directory
 	FString Arguments = FString::Printf(TEXT("/register \"%s\""), *EngineRootDir);
-	return FPlatformProcess::ExecElevatedProcess(*VersionSelectorFileName, *Arguments, NULL);
+	int32 ExitCode;
+	if (!FPlatformProcess::ExecElevatedProcess(*VersionSelectorFileName, *Arguments, &ExitCode) || ExitCode != 0)
+	{
+		return false;
+	}
+
+	// Notify the user that everything is awesome.
+	FPlatformMisc::MessageBoxExt(EAppMsgType::Ok, TEXT("Registration successful."), TEXT("Success"));
+	return true;
 }
 
 void EnumerateInvalidInstallations(TArray<FString> &OutIds)
 {
 	// Enumerate all the engine installations
-	TMap<FString, FString> Installations;
-	FPlatformInstallation::EnumerateEngineInstallations(Installations);
+	TArray< TPair<FString, FString> > Installations;
+	FPlatformMisc::EnumerateEngineInstallations(Installations);
 
 	// Check each one of them has a valid engine path
-	for (TMap<FString, FString>::TConstIterator Iter(Installations); Iter; ++Iter)
+	for (TArray<TPair<FString, FString>>::TConstIterator Iter(Installations); Iter; ++Iter)
 	{
-		FString RootDir = Iter.Value();
+		const FString &RootDir = Iter->Value;
 		if (!FPlatformInstallation::IsValidEngineRootDir(RootDir))
 		{
-			OutIds.Add(Iter.Key());
+			OutIds.Add(Iter->Key);
 		}
 	}
 }
@@ -103,7 +111,7 @@ bool RegisterEngineDirectory(const FString &EngineDir)
 {
 	// Get any existing tag name or register a new one
 	FString EngineId;
-	if (!GetEngineIdFromRootDir(EngineDir, EngineId))
+	if (!FPlatformMisc::GetEngineIdentifierFromRootDir(EngineDir, EngineId))
 	{
 		EngineId = FGuid::NewGuid().ToString(EGuidFormats::DigitsWithHyphensInBraces);
 	}
@@ -145,14 +153,14 @@ bool BrowseForEngineAssociation(const FString &ProjectFileName)
 
 	// Check that it's a registered engine directory
 	FString EngineId;
-	if (!GetEngineIdFromRootDir(NewEngineRootDir, EngineId))
+	if (!FPlatformMisc::GetEngineIdentifierFromRootDir(NewEngineRootDir, EngineId))
 	{
 		int32 ExitCode;
 		if (!FPlatformProcess::ExecElevatedProcess(*FullExecutableName(), *FString::Printf(TEXT("/register \"%s\""), *NewEngineRootDir), &ExitCode) || ExitCode != 0)
 		{
 			return false;
 		}
-		if (!GetEngineIdFromRootDir(NewEngineRootDir, EngineId))
+		if (!FPlatformMisc::GetEngineIdentifierFromRootDir(NewEngineRootDir, EngineId))
 		{
 			return false;
 		}
diff --git a/Engine/Source/Programs/UnrealVersionSelector/Private/Windows/WindowsPlatformInstallation.cpp b/Engine/Source/Programs/UnrealVersionSelector/Private/Windows/WindowsPlatformInstallation.cpp
--- a/Engine/Source/Programs/UnrealVersionSelector/Private/Windows/WindowsPlatformInstallation.cpp
+++ b/Engine/Source/Programs/UnrealVersionSelector/Private/Windows/WindowsPlatformInstallation.cpp
@@ -110,8 +110,8 @@ void GetRequiredRegistrySettings(TIndirectArray<FRegistryRootedKey> &RootedKeys)
 	FRegistryKey *SwitchShellKey = SwitchKey->FindOrAddKey(TEXT("shell"));
 
 	// HKLM\SOFTWARE\Classes\Unreal.ProjectFile\switchversion\shell\<Label>
-	TMap<FString, FString> Installations;
-	FPlatformInstallation::EnumerateEngineInstallations(Installations);
+	TArray<TPair<FString, FString>> Installations;
+	FPlatformMisc::EnumerateEngineInstallations(Installations);
 
 	if (Installations.Num() > 0)
 	{
@@ -122,10 +122,10 @@ void GetRequiredRegistrySettings(TIndirectArray<FRegistryRootedKey> &RootedKeys)
 
 		// Build a list of installation names
 		TMap<FString, FString> InstallationNames;
-		for (TMap<FString, FString>::TConstIterator Iter(Installations); Iter; ++Iter)
+		for (TArray<TPair<FString, FString>>::TConstIterator Iter(Installations); Iter; ++Iter)
 		{
-			const FString &Id = Iter.Key();
-			FString Description = GetInstallationDescription(Id, Iter.Value());
+			const FString &Id = Iter->Key;
+			FString Description = GetInstallationDescription(Id, Iter->Value);
 			InstallationNames.Add(Description, Id);
 		}
 		InstallationNames.KeySort(FEngineLabelSortPredicate());
@@ -181,9 +181,12 @@ bool FWindowsPlatformInstallation::GetLauncherVersionSelector(FString &OutFileNa
 
 bool FWindowsPlatformInstallation::RegisterEngineInstallation(const FString &Id, const FString &RootDirName)
 {
+	FString NormalizedRootDirName = RootDirName;
+	FPaths::NormalizeDirectoryName(NormalizedRootDirName);
+
 	FRegistryRootedKey RootKey(HKEY_LOCAL_MACHINE, FString::Printf(TEXT("SOFTWARE\\EpicGames\\Unreal Engine\\%s"), *Id));
 	RootKey.Key = new FRegistryKey();
-	RootKey.Key->SetValue(TEXT("InstalledDirectory"), RootDirName);
+	RootKey.Key->SetValue(TEXT("InstalledDirectory"), NormalizedRootDirName);
 	return RootKey.Write(false);
 }
 
@@ -193,31 +196,6 @@ bool FWindowsPlatformInstallation::UnregisterEngineInstallation(const FString &I
 	return RootKey.Write(true);
 }
 
-void FWindowsPlatformInstallation::EnumerateEngineInstallations(TMap<FString, FString> &OutInstallations)
-{
-	HKEY hKey;
-	if (RegOpenKeyEx(HKEY_LOCAL_MACHINE, L"SOFTWARE\\EpicGames\\Unreal Engine", 0, KEY_READ, &hKey) == ERROR_SUCCESS)
-	{
-		for (DWORD Index = 0;; Index++)
-		{
-			wchar_t KeyName[256];
-			DWORD KeyLength = sizeof(KeyName) / sizeof(KeyName[0]);
-			if (RegEnumKeyEx(hKey, Index, KeyName, &KeyLength, NULL, NULL, NULL, NULL) != ERROR_SUCCESS) break;
-
-			wchar_t InstalledDirectory[MAX_PATH];
-			DWORD InstalledDirectoryLen = sizeof(InstalledDirectory) / sizeof(InstalledDirectory[0]);
-
-			if (RegGetValue(hKey, KeyName, L"InstalledDirectory", RRF_RT_REG_SZ, NULL, InstalledDirectory, &InstalledDirectoryLen) == ERROR_SUCCESS)
-			{
-				FString NormalizedInstalledDirectory = InstalledDirectory;
-				FPaths::NormalizeDirectoryName(NormalizedInstalledDirectory);
-				OutInstallations.Add(KeyName, NormalizedInstalledDirectory);
-			}
-		}
-		RegCloseKey(hKey);
-	}
-}
-
 bool FWindowsPlatformInstallation::IsSourceDistribution(const FString &EngineRootDir)
 {
 	// Check for the existence of a GenerateProjectFiles.bat file. This allows compatibility with the GitHub 4.0 release.
diff --git a/Engine/Source/Programs/UnrealVersionSelector/Private/Windows/WindowsPlatformInstallation.h b/Engine/Source/Programs/UnrealVersionSelector/Private/Windows/WindowsPlatformInstallation.h
--- a/Engine/Source/Programs/UnrealVersionSelector/Private/Windows/WindowsPlatformInstallation.h
+++ b/Engine/Source/Programs/UnrealVersionSelector/Private/Windows/WindowsPlatformInstallation.h
@@ -10,10 +10,11 @@ struct FWindowsPlatformInstallation : FGenericPlatformInstallation
 	/// Get the path to the installed build of the version selector
 	static bool GetLauncherVersionSelector(FString &OutFileName);
 
-	// Engine installations
+	// Register an engine installation
 	static bool RegisterEngineInstallation(const FString &Id, const FString &RootDirName);
+
+	// Unregister an engine installation
 	static bool UnregisterEngineInstallation(const FString &Id);
-	static void EnumerateEngineInstallations(TMap<FString, FString> &OutInstallations);
 
 	// Test whether a given engine root directory is a source build
 	static bool IsSourceDistribution(const FString &EngineRootDir);
diff --git a/Engine/Source/Runtime/Core/Private/GenericPlatform/GenericPlatformMisc.cpp b/Engine/Source/Runtime/Core/Private/GenericPlatform/GenericPlatformMisc.cpp
--- a/Engine/Source/Runtime/Core/Private/GenericPlatform/GenericPlatformMisc.cpp
+++ b/Engine/Source/Runtime/Core/Private/GenericPlatform/GenericPlatformMisc.cpp
@@ -12,6 +12,8 @@
 #include "ModuleManager.h"
 #include "VarargsHelper.h"
 #include "SecureHash.h"
+#include "Containers/Map.h"
+#include "../../Launch/Resources/Version.h"
 
 #include "UProjectInfo.h"
 
@@ -686,3 +688,58 @@ FText FGenericPlatformMisc::GetFileManagerName()
 {
 	return NSLOCTEXT("GenericPlatform", "FileManagerName", "File Manager");
 }
+
+void FGenericPlatformMisc::EnumerateEngineInstallations(TArray< TPair<FString, FString> > &OutInstallations)
+{
+	// Read the config file.
+	FConfigFile ConfigFile;
+	ConfigFile.Read(FString(FPlatformProcess::ApplicationSettingsDir()) / FString(EPIC_PRODUCT_IDENTIFIER) / FString(TEXT("Install.ini")));
+
+	// Enumerate the installation list.
+	FConfigSection &Section = ConfigFile.FindOrAdd(TEXT("Installations"));
+	for (TMultiMap<FName, FString>::TConstIterator Iter(Section); Iter; ++Iter)
+	{
+		OutInstallations.Add(TPairInitializer<FString, FString>(Iter.Key().ToString(), Iter.Value()));
+	}
+}
+
+bool FGenericPlatformMisc::GetEngineRootDirFromIdentifier(const FString &Identifier, FString &OutRootDir)
+{
+	// Get all the installations
+	TArray< TPair<FString, FString> > Installations;
+	FPlatformMisc::EnumerateEngineInstallations(Installations);
+
+	// Find the one with the right identifier
+	for (TArray< TPair<FString, FString> >::TConstIterator Iter(Installations); Iter; ++Iter)
+	{
+		if (Iter->Key == Identifier)
+		{
+			OutRootDir = Iter->Value;
+			return true;
+		}
+	}
+	return false;
+}
+
+bool FGenericPlatformMisc::GetEngineIdentifierFromRootDir(const FString &RootDir, FString &OutIdentifier)
+{
+	// Get all the installations
+	TArray< TPair<FString, FString> > Installations;
+	FPlatformMisc::EnumerateEngineInstallations(Installations);
+
+	// Normalize the root directory
+	FString NormalizedRootDir = RootDir;
+	FPaths::NormalizeDirectoryName(NormalizedRootDir);
+
+	// Find the label for the given directory
+	for (TArray< TPair<FString, FString> >::TConstIterator Iter(Installations); Iter; ++Iter)
+	{
+		if (Iter->Value == NormalizedRootDir)
+		{
+			OutIdentifier = Iter->Key;
+			return true;
+		}
+	}
+	return false;
+}
+
diff --git a/Engine/Source/Runtime/Core/Private/Windows/WindowsPlatformMisc.cpp b/Engine/Source/Runtime/Core/Private/Windows/WindowsPlatformMisc.cpp
--- a/Engine/Source/Runtime/Core/Private/Windows/WindowsPlatformMisc.cpp
+++ b/Engine/Source/Runtime/Core/Private/Windows/WindowsPlatformMisc.cpp
@@ -1847,3 +1847,28 @@ FText FWindowsPlatformMisc::GetFileManagerName()
 {
 	return NSLOCTEXT("WindowsPlatform", "FileManagerName", "Explorer");
 }
+
+void FWindowsPlatformMisc::EnumerateEngineInstallations(TArray< TPair<FString, FString> > &OutInstallations)
+{
+	HKEY hKey;
+	if (RegOpenKeyEx(HKEY_LOCAL_MACHINE, TEXT("SOFTWARE\\EpicGames\\Unreal Engine"), 0, KEY_READ, &hKey) == ERROR_SUCCESS)
+	{
+		for (::DWORD Index = 0;; Index++)
+		{
+			TCHAR KeyName[256];
+			::DWORD KeyLength = sizeof(KeyName) / sizeof(KeyName[0]);
+			if (RegEnumKeyEx(hKey, Index, KeyName, &KeyLength, NULL, NULL, NULL, NULL) != ERROR_SUCCESS) break;
+
+			TCHAR InstalledDirectory[MAX_PATH];
+			::DWORD InstalledDirectoryLen = sizeof(InstalledDirectory) / sizeof(InstalledDirectory[0]);
+
+			if (RegGetValue(hKey, KeyName, TEXT("InstalledDirectory"), RRF_RT_REG_SZ, NULL, InstalledDirectory, &InstalledDirectoryLen) == ERROR_SUCCESS)
+			{
+				FString NormalizedInstalledDirectory = InstalledDirectory;
+				FPaths::NormalizeDirectoryName(NormalizedInstalledDirectory);
+				OutInstallations.Add(TPairInitializer<FString, FString>(KeyName, NormalizedInstalledDirectory));
+			}
+		}
+		RegCloseKey(hKey);
+	}
+}
diff --git a/Engine/Source/Runtime/Core/Public/GenericPlatform/GenericPlatformMisc.h b/Engine/Source/Runtime/Core/Public/GenericPlatform/GenericPlatformMisc.h
--- a/Engine/Source/Runtime/Core/Public/GenericPlatform/GenericPlatformMisc.h
+++ b/Engine/Source/Runtime/Core/Public/GenericPlatform/GenericPlatformMisc.h
@@ -9,6 +9,8 @@
 
 class FText;
 
+template<typename KeyType, typename ValueType> class TPair;
+
 namespace EBuildConfigurations
 {
 	/**
@@ -721,6 +723,32 @@ public:
 		return PLATFORM_HAS_TOUCH_MAIN_SCREEN;
 	}
 
+#if PLATFORM_DESKTOP
+	/**
+	* Enumerates all the registered engine installations.
+	*
+	* @param	OutInstallations	Array which is filled in with identifier/root-directory pairs for all known installations. Identifiers are typically
+	*								version strings for canonical UE4 releases or GUID strings for GitHub releases.
+	*/
+	static void EnumerateEngineInstallations(TArray< TPair<FString, FString> > &OutInstallations);
+
+	/**
+	* Returns the identifier for the engine with the given root directory.
+	*
+	* @param	RootDirName			Root directory for the engine.
+	* @param	OutIdentifier		Identifier used to refer to this installation.
+	*/
+	static bool GetEngineRootDirFromIdentifier(const FString &Identifier, FString &OutRootDir);
+
+	/**
+	* Returns the identifier for the engine with the given root directory.
+	*
+	* @param	RootDirName			Root directory for the engine.
+	* @param	OutIdentifier		Identifier used to refer to this installation.
+	*/
+	static bool GetEngineIdentifierFromRootDir(const FString &RootDir, FString &OutIdentifier);
+#endif
+
 #if !UE_BUILD_SHIPPING
 protected:
 	/** Whether the user should be prompted to allow for a remote debugger to be attached */
diff --git a/Engine/Source/Runtime/Core/Public/Windows/WindowsPlatformMisc.h b/Engine/Source/Runtime/Core/Public/Windows/WindowsPlatformMisc.h
--- a/Engine/Source/Runtime/Core/Public/Windows/WindowsPlatformMisc.h
+++ b/Engine/Source/Runtime/Core/Public/Windows/WindowsPlatformMisc.h
@@ -188,6 +188,9 @@ struct CORE_API FWindowsPlatformMisc : public FGenericPlatformMisc
 
 	/** @return Get the name of the platform specific file manager (Explorer) */
 	static FText GetFileManagerName();
+
+	/** Enumerate all the engine installations, using the registry */
+	static void EnumerateEngineInstallations(TArray< TPair<FString, FString> > &OutInstallations);
 };
 
 typedef FWindowsPlatformMisc FPlatformMisc;
