Fix command line for handling of GenerateProjectFiles.bat in UnrealVersionSelector

diff --git a/Engine/Source/Programs/UnrealVersionSelector/Private/UnrealVersionSelector.cpp b/Engine/Source/Programs/UnrealVersionSelector/Private/UnrealVersionSelector.cpp
--- a/Engine/Source/Programs/UnrealVersionSelector/Private/UnrealVersionSelector.cpp
+++ b/Engine/Source/Programs/UnrealVersionSelector/Private/UnrealVersionSelector.cpp
@@ -220,14 +220,14 @@ bool GenerateProjectFiles(const FString &ProjectFileName)
 	}
 
 	// Build the argument list
-	FString Arguments = FString::Printf(TEXT("\"%s\" -game"), *ProjectFileName);
+	FString Arguments = TEXT("-game");
 	if (FPlatformInstallation::IsSourceDistribution(RootDir))
 	{
 		Arguments += TEXT(" -engine");
 	}
 
 	// Launch the editor
-	if (!FPlatformInstallation::LaunchEditor(RootDir, FString::Printf(TEXT("\"%s\" %s"), *ProjectFileName, *Arguments)))
+	if (!FPlatformInstallation::GenerateProjectFiles(RootDir, FString::Printf(TEXT("\"%s\" %s"), *ProjectFileName, *Arguments)))
 	{
 		FPlatformMisc::MessageBoxExt(EAppMsgType::Ok, TEXT("Failed to launch editor"), TEXT("Error"));
 		return false;
@@ -284,7 +284,7 @@ int Main(const TArray<FString> &Arguments)
 		// Play a game using the editor executable
 		bRes = LaunchEditor(Arguments[1], L"/game");
 	}
-	else if (Arguments.Num() == 2 && Arguments[0] == TEXT("/generateprojectfiles"))
+	else if (Arguments.Num() == 2 && Arguments[0] == TEXT("/projectfiles"))
 	{
 		// Generate Visual Studio project files
 		bRes = GenerateProjectFiles(Arguments[1]);
diff --git a/Engine/Source/Programs/UnrealVersionSelector/Private/Windows/WindowsPlatformInstallation.cpp b/Engine/Source/Programs/UnrealVersionSelector/Private/Windows/WindowsPlatformInstallation.cpp
--- a/Engine/Source/Programs/UnrealVersionSelector/Private/Windows/WindowsPlatformInstallation.cpp
+++ b/Engine/Source/Programs/UnrealVersionSelector/Private/Windows/WindowsPlatformInstallation.cpp
@@ -217,7 +217,7 @@ bool FWindowsPlatformInstallation::LaunchEditor(const FString &RootDirName, cons
 
 bool FWindowsPlatformInstallation::GenerateProjectFiles(const FString &RootDirName, const FString &Arguments)
 {
-	FString AllArguments = FString::Printf(TEXT("/c \"%s\" %s"), *(RootDirName / TEXT("Engine/Build/BatchFiles/RocketGenerateProjectFiles.bat")), *Arguments);
+	FString AllArguments = FString::Printf(TEXT("/c \"\"%s\" %s\""), *(RootDirName / TEXT("Engine/Build/BatchFiles/GenerateProjectFiles.bat")), *Arguments);
 	return FPlatformProcess::ExecProcess(TEXT("cmd.exe"), *AllArguments, NULL, NULL, NULL);
 }
 
