Updates to version selector.

* Lots of code moved into IDesktopPlatform for sharing with Launcher and Mac (including setting up file associations, querying project versions, etc...)
* Hack to enumerate all the known launcher engine installations. Does not use registry keys any more. Will probably change to use a list of installations generated by the launcher at some point soon.
* List of registered GitHub builds is stored in HKEY_CURRENT_USER
* Switching engine versions is now done through a dialog rather than through the context menu.
* VersionSelector includes a version number for shell integration, allowing it to defer to an existing installation of the same version if necessary.

diff --git a/Engine/Source/Developer/DesktopPlatform/Private/DesktopPlatformBase.cpp b/Engine/Source/Developer/DesktopPlatform/Private/DesktopPlatformBase.cpp
--- a/Engine/Source/Developer/DesktopPlatform/Private/DesktopPlatformBase.cpp
+++ b/Engine/Source/Developer/DesktopPlatform/Private/DesktopPlatformBase.cpp
@@ -9,6 +9,13 @@ FString FDesktopPlatformBase::GetCurrentEngineIdentifier()
 	return GetEngineIdentifierFromRootDir(FPlatformMisc::RootDir(), Identifier) ? Identifier : TEXT("");
 }
 
+void FDesktopPlatformBase::EnumerateLauncherEngineInstallations(TMap<FString, FString> &OutInstallations)
+{
+	// Just enumerate the known manifests for the time being 
+	CheckForLauncherEngineInstallation(TEXT("40003"), TEXT("4.0"), OutInstallations);
+	CheckForLauncherEngineInstallation(TEXT("1040003"), TEXT("4.1"), OutInstallations);
+}
+
 bool FDesktopPlatformBase::GetEngineRootDirFromIdentifier(const FString &Identifier, FString &OutRootDir)
 {
 	// Get all the installations
@@ -48,3 +55,173 @@ bool FDesktopPlatformBase::GetEngineIdentifierFromRootDir(const FString &RootDir
 	}
 	return false;
 }
+
+bool FDesktopPlatformBase::GetDefaultEngineIdentifier(FString &OutId)
+{
+	TMap<FString, FString> Installations;
+	EnumerateEngineInstallations(Installations);
+
+	bool bRes = false;
+	if (Installations.Num() > 0)
+	{
+		// Default to the first install
+		TMap<FString, FString>::TConstIterator Iter(Installations);
+		OutId = Iter.Key();
+		++Iter;
+
+		// Try to find the most preferred install
+		for(; Iter; ++Iter)
+		{
+			if(IsPreferredEngineIdentifier(Iter.Key(), OutId))
+			{
+				OutId = Iter.Key();
+			}
+		}
+	}
+	return bRes;
+}
+
+bool FDesktopPlatformBase::GetDefaultEngineRootDir(FString &OutDirName)
+{
+	FString Identifier;
+	return GetDefaultEngineIdentifier(Identifier) && GetEngineRootDirFromIdentifier(Identifier, OutDirName);
+}
+
+bool FDesktopPlatformBase::IsPreferredEngineIdentifier(const FString &Identifier, const FString &OtherIdentifier)
+{
+	int32 Version = ParseReleaseVersion(Identifier);
+	int32 OtherVersion = ParseReleaseVersion(OtherIdentifier);
+
+	if(Version != OtherVersion)
+	{
+		return Version > OtherVersion;
+	}
+	else
+	{
+		return Identifier > OtherIdentifier;
+	}
+}
+
+bool FDesktopPlatformBase::IsSourceDistribution(const FString &EngineRootDir)
+{
+	// Check for the existence of a SourceBuild.txt file
+	FString SourceBuildPath = EngineRootDir / TEXT("Engine\\Build\\SourceDistribution.txt");
+	return (IFileManager::Get().FileSize(*SourceBuildPath) >= 0);
+}
+
+bool FDesktopPlatformBase::IsPerforceBuild(const FString &EngineRootDir)
+{
+	// Check for the existence of a SourceBuild.txt file
+	FString PerforceBuildPath = EngineRootDir / TEXT("Engine\\Build\\PerforceBuild.txt");
+	return (IFileManager::Get().FileSize(*PerforceBuildPath) >= 0);
+}
+
+bool FDesktopPlatformBase::IsValidRootDirectory(const FString &RootDir)
+{
+	FString EngineBinariesDirName = RootDir / TEXT("Engine/Binaries");
+	FPaths::NormalizeDirectoryName(EngineBinariesDirName);
+	return IFileManager::Get().DirectoryExists(*EngineBinariesDirName);
+}
+
+bool FDesktopPlatformBase::SetEngineIdentifierForProject(const FString &ProjectFileName, const FString &Identifier)
+{
+	TSharedPtr<FJsonObject> ProjectFile = LoadProjectFile(ProjectFileName);
+	if (!ProjectFile.IsValid())
+	{
+		return false;
+	}
+
+	ProjectFile->SetStringField(TEXT("EngineAssociation"), Identifier);
+	return SaveProjectFile(ProjectFileName, ProjectFile);
+}
+
+bool FDesktopPlatformBase::GetEngineIdentifierForProject(const FString &ProjectFileName, FString &OutIdentifier)
+{
+	TSharedPtr<FJsonObject> ProjectFile = LoadProjectFile(ProjectFileName);
+	if(!ProjectFile.IsValid())
+	{
+		return false;
+	}
+
+	OutIdentifier = ProjectFile->GetStringField(TEXT("EngineAssociation"));
+	return true;
+}
+
+void FDesktopPlatformBase::CheckForLauncherEngineInstallation(const FString &AppId, const FString &Identifier, TMap<FString, FString> &OutInstallations)
+{
+	FString ManifestText;
+	FString ManifestFileName = FString(FPlatformProcess::ApplicationSettingsDir()) / FString::Printf(TEXT("UnrealEngineLauncher/Data/Manifests/%s.manifest"), *AppId);
+	if (FFileHelper::LoadFileToString(ManifestText, *ManifestFileName))
+	{
+		TSharedPtr< FJsonObject > RootObject;
+		TSharedRef< TJsonReader<> > Reader = TJsonReaderFactory<>::Create(ManifestText);
+		if (FJsonSerializer::Deserialize(Reader, RootObject) && RootObject.IsValid())
+		{
+			TSharedPtr<FJsonObject> CustomFieldsObject = RootObject->GetObjectField(TEXT("CustomFields"));
+			if (CustomFieldsObject.IsValid())
+			{
+				FString InstallLocation = CustomFieldsObject->GetStringField("InstallLocation");
+				if (InstallLocation.Len() > 0)
+				{
+					OutInstallations.Add(Identifier, InstallLocation);
+				}
+			}
+		}
+	}
+}
+
+int32 FDesktopPlatformBase::ParseReleaseVersion(const FString &Version)
+{
+	TCHAR *End;
+
+	uint64 Major = FCString::Strtoui64(*Version, &End, 10);
+	if (Major >= MAX_int16 || *(End++) != '.')
+	{
+		return INDEX_NONE;
+	}
+
+	uint64 Minor = FCString::Strtoui64(End, &End, 10);
+	if (Minor >= MAX_int16 || *End != 0)
+	{
+		return INDEX_NONE;
+	}
+
+	return (Major << 16) + Minor;
+}
+
+TSharedPtr<FJsonObject> FDesktopPlatformBase::LoadProjectFile(const FString &FileName)
+{
+	FString FileContents;
+
+	if (!FFileHelper::LoadFileToString(FileContents, *FileName))
+	{
+		return TSharedPtr<FJsonObject>(NULL);
+	}
+
+	TSharedPtr< FJsonObject > JsonObject;
+	TSharedRef< TJsonReader<> > Reader = TJsonReaderFactory<>::Create(FileContents);
+	if (!FJsonSerializer::Deserialize(Reader, JsonObject) || !JsonObject.IsValid())
+	{
+		return TSharedPtr<FJsonObject>(NULL);
+	}
+
+	return JsonObject;
+}
+
+bool FDesktopPlatformBase::SaveProjectFile(const FString &FileName, TSharedPtr<FJsonObject> Object)
+{
+	FString FileContents;
+
+	TSharedRef< TJsonWriter<> > Writer = TJsonWriterFactory<>::Create(&FileContents);
+	if (!FJsonSerializer::Serialize(Object.ToSharedRef(), Writer))
+	{
+		return false;
+	}
+
+	if (!FFileHelper::SaveStringToFile(FileContents, *FileName))
+	{
+		return false;
+	}
+
+	return true;
+}
diff --git a/Engine/Source/Developer/DesktopPlatform/Private/DesktopPlatformBase.h b/Engine/Source/Developer/DesktopPlatform/Private/DesktopPlatformBase.h
--- a/Engine/Source/Developer/DesktopPlatform/Private/DesktopPlatformBase.h
+++ b/Engine/Source/Developer/DesktopPlatform/Private/DesktopPlatformBase.h
@@ -3,12 +3,32 @@
 #pragma once
 
 #include "IDesktopPlatform.h"
+#include "Runtime/Core/Public/Serialization/Json/Json.h"
 
 class FDesktopPlatformBase : public IDesktopPlatform
 {
 public:
 	// IDesktopPlatform Implementation
 	virtual FString GetCurrentEngineIdentifier() OVERRIDE;
+	virtual void EnumerateLauncherEngineInstallations(TMap<FString, FString> &OutInstallations) OVERRIDE;
 	virtual bool GetEngineRootDirFromIdentifier(const FString &Identifier, FString &OutRootDir) OVERRIDE;
 	virtual bool GetEngineIdentifierFromRootDir(const FString &RootDir, FString &OutIdentifier) OVERRIDE;
+
+	virtual bool GetDefaultEngineIdentifier(FString &OutIdentifier) OVERRIDE;
+	virtual bool GetDefaultEngineRootDir(FString &OutRootDir) OVERRIDE;
+	virtual bool IsPreferredEngineIdentifier(const FString &Identifier, const FString &OtherIdentifier);
+
+	virtual bool IsSourceDistribution(const FString &RootDir) OVERRIDE;
+	virtual bool IsPerforceBuild(const FString &RootDir) OVERRIDE;
+	virtual bool IsValidRootDirectory(const FString &RootDir) OVERRIDE;
+
+	virtual bool SetEngineIdentifierForProject(const FString &ProjectFileName, const FString &Identifier) OVERRIDE;
+	virtual bool GetEngineIdentifierForProject(const FString &ProjectFileName, FString &OutIdentifier) OVERRIDE;
+
+private:
+	void CheckForLauncherEngineInstallation(const FString &AppId, const FString &Identifier, TMap<FString, FString> &OutInstallations);
+	int32 ParseReleaseVersion(const FString &Version);
+
+	TSharedPtr<FJsonObject> LoadProjectFile(const FString &FileName);
+	bool FDesktopPlatformBase::SaveProjectFile(const FString &FileName, TSharedPtr<FJsonObject> Object);
 };
diff --git a/Engine/Source/Developer/DesktopPlatform/Private/Mac/DesktopPlatformMac.cpp b/Engine/Source/Developer/DesktopPlatform/Private/Mac/DesktopPlatformMac.cpp
--- a/Engine/Source/Developer/DesktopPlatform/Private/Mac/DesktopPlatformMac.cpp
+++ b/Engine/Source/Developer/DesktopPlatform/Private/Mac/DesktopPlatformMac.cpp
@@ -514,6 +514,17 @@ bool FDesktopPlatformMac::FileDialogShared(bool bSave, const void* ParentWindowH
 
 void FDesktopPlatformMac::EnumerateEngineInstallations(TMap<FString, FString> &OutInstallations)
 {
+	EnumerateLauncherEngineInstallations(OutInstallations);
+}
+
+bool FDesktopPlatformMac::VerifyFileAssociations() OVERRIDE
+{
+	return true;
+}
+
+bool FDesktopPlatformMac::UpdateFileAssociations() OVERRIDE
+{
+	return true;
 }
 
 #undef LOCTEXT_NAMESPACE
diff --git a/Engine/Source/Developer/DesktopPlatform/Private/Mac/DesktopPlatformMac.h b/Engine/Source/Developer/DesktopPlatform/Private/Mac/DesktopPlatformMac.h
--- a/Engine/Source/Developer/DesktopPlatform/Private/Mac/DesktopPlatformMac.h
+++ b/Engine/Source/Developer/DesktopPlatform/Private/Mac/DesktopPlatformMac.h
@@ -16,6 +16,9 @@ public:
 
 	virtual void EnumerateEngineInstallations(TMap<FString, FString> &OutInstallations) OVERRIDE;
 
+	virtual bool VerifyFileAssociations() OVERRIDE;
+	virtual bool UpdateFileAssociations() OVERRIDE;
+
 private:
 	bool FileDialogShared(bool bSave, const void* ParentWindowHandle, const FString& DialogTitle, const FString& DefaultPath, const FString& DefaultFile, const FString& FileTypes, uint32 Flags, TArray<FString>& OutFilenames);
 };
diff --git a/Engine/Source/Developer/DesktopPlatform/Private/Windows/DesktopPlatformWindows.cpp b/Engine/Source/Developer/DesktopPlatform/Private/Windows/DesktopPlatformWindows.cpp
--- a/Engine/Source/Developer/DesktopPlatform/Private/Windows/DesktopPlatformWindows.cpp
+++ b/Engine/Source/Developer/DesktopPlatform/Private/Windows/DesktopPlatformWindows.cpp
@@ -347,27 +347,189 @@ bool FDesktopPlatformWindows::FileDialogShared(bool bSave, const void* ParentWin
 
 void FDesktopPlatformWindows::EnumerateEngineInstallations(TMap<FString, FString> &OutInstallations)
 {
+	// Enumerate the binary installations
+	EnumerateLauncherEngineInstallations(OutInstallations);
+
+	// Enumerate the per-user installations
 	HKEY hKey;
-	if (RegOpenKeyEx(HKEY_LOCAL_MACHINE, TEXT("SOFTWARE\\EpicGames\\Unreal Engine"), 0, KEY_READ, &hKey) == ERROR_SUCCESS)
+	if (RegOpenKeyEx(HKEY_CURRENT_USER, TEXT("SOFTWARE\\Epic Games\\Unreal Engine\\Builds"), 0, KEY_ALL_ACCESS, &hKey) == ERROR_SUCCESS)
 	{
+		// Enumerate all the installations
 		for (::DWORD Index = 0;; Index++)
 		{
-			TCHAR KeyName[256];
-			::DWORD KeyLength = sizeof(KeyName) / sizeof(KeyName[0]);
-			if (RegEnumKeyEx(hKey, Index, KeyName, &KeyLength, NULL, NULL, NULL, NULL) != ERROR_SUCCESS) break;
+			TCHAR ValueName[256];
+			TCHAR ValueData[MAX_PATH];
+			::DWORD ValueType = 0;
+			::DWORD ValueNameLength = sizeof(ValueName) / sizeof(ValueName[0]);
+			::DWORD ValueDataSize = sizeof(ValueData);
+
+			LRESULT Result = RegEnumValue(hKey, Index, ValueName, &ValueNameLength, NULL, &ValueType, (BYTE*)&ValueData[0], &ValueDataSize);
+			if(Result == ERROR_SUCCESS)
+			{
+				int32 ValueDataLength = ValueDataSize / sizeof(TCHAR);
+				if(ValueDataLength > 0 && ValueData[ValueDataLength - 1] == 0) ValueDataLength--;
 
-			TCHAR InstalledDirectory[MAX_PATH];
-			::DWORD InstalledDirectoryLen = sizeof(InstalledDirectory) / sizeof(InstalledDirectory[0]);
+				FString NormalizedInstalledDirectory(ValueDataLength, ValueData);
+				FPaths::NormalizeDirectoryName(NormalizedInstalledDirectory);
+				OutInstallations.Add(ValueName, NormalizedInstalledDirectory);
+			}
+			else if(Result == ERROR_NO_MORE_ITEMS)
+			{
+				break;
+			}
+		}
 
-			if (RegGetValue(hKey, KeyName, TEXT("InstalledDirectory"), RRF_RT_REG_SZ, NULL, InstalledDirectory, &InstalledDirectoryLen) == ERROR_SUCCESS)
+		// Trim the list to everything that's valid
+		for(TMap<FString, FString>::TIterator Iter(OutInstallations); Iter; ++Iter)
+		{
+			if (!IsValidRootDirectory(Iter.Value()))
 			{
-				FString NormalizedInstalledDirectory = InstalledDirectory;
-				FPaths::NormalizeDirectoryName(NormalizedInstalledDirectory);
-				OutInstallations.Add(KeyName, NormalizedInstalledDirectory);
+				RegDeleteValue(hKey, *Iter.Key());
+				Iter.RemoveCurrent();
 			}
 		}
+
 		RegCloseKey(hKey);
 	}
 }
 
+bool FDesktopPlatformWindows::IsSourceDistribution(const FString &RootDir)
+{
+	// Check for the existence of a GenerateProjectFiles.bat file. This allows compatibility with the GitHub 4.0 release.
+	FString GenerateProjectFilesPath = RootDir / TEXT("GenerateProjectFiles.bat");
+	if (IFileManager::Get().FileSize(*GenerateProjectFilesPath) >= 0)
+	{
+		return true;
+	}
+
+	// Otherwise use the default test
+	return FDesktopPlatformBase::IsSourceDistribution(RootDir);
+}
+
+bool FDesktopPlatformWindows::VerifyFileAssociations()
+{
+	TIndirectArray<FRegistryRootedKey> Keys;
+	GetRequiredRegistrySettings(Keys);
+
+	for (int32 Idx = 0; Idx < Keys.Num(); Idx++)
+	{
+		if (!Keys[Idx].IsUpToDate(true))
+		{
+			return false;
+		}
+	}
+
+	return true;
+}
+
+bool FDesktopPlatformWindows::UpdateFileAssociations()
+{
+	TIndirectArray<FRegistryRootedKey> Keys;
+	GetRequiredRegistrySettings(Keys);
+
+	for (int32 Idx = 0; Idx < Keys.Num(); Idx++)
+	{
+		if (!Keys[Idx].Write(true))
+		{
+			return false;
+		}
+	}
+
+	return true;
+}
+
+void FDesktopPlatformWindows::GetRequiredRegistrySettings(TIndirectArray<FRegistryRootedKey> &RootedKeys)
+{
+	// Get the path to VersionSelector.exe. If we're running from UnrealVersionSelector itself, try to stick with the current configuration.
+	FString DefaultVersionSelectorName = FPlatformProcess::ExecutableName(false);
+	if (!DefaultVersionSelectorName.StartsWith(TEXT("UnrealVersionSelector")))
+	{
+		DefaultVersionSelectorName = TEXT("UnrealVersionSelector-Win64-Shipping.exe");
+	}
+	FString ExecutableFileName = FString(FPlatformProcess::BaseDir()) / DefaultVersionSelectorName;
+
+	// Defer to UnrealVersionSelector.exe in a launcher installation if it's got the same version number or greater.
+	TCHAR InstallDir[MAX_PATH];
+	::DWORD InstallDirSize = sizeof(InstallDir);
+	if(RegGetValue(HKEY_LOCAL_MACHINE, TEXT("SOFTWARE\\EpicGames\\Unreal Engine"), TEXT("INSTALLDIR"), RRF_RT_REG_SZ, NULL, InstallDir, &InstallDirSize) == ERROR_SUCCESS)
+	{
+		FString NormalizedInstallDir = InstallDir;
+		FPaths::NormalizeDirectoryName(NormalizedInstallDir);
+
+		FString InstalledExecutableFileName = NormalizedInstallDir / FString("Launcher/Engine/Binaries/Win64/UnrealVersionSelector-Win64-Shipping.exe");
+		if(GetShellIntegrationVersion(InstalledExecutableFileName) == GetShellIntegrationVersion(ExecutableFileName))
+		{
+			ExecutableFileName = InstalledExecutableFileName;
+		}
+	}
+
+	// Get the path to the executable
+	FPaths::MakePlatformFilename(ExecutableFileName);
+	FString QuotedExecutableFileName = FString(TEXT("\"")) + ExecutableFileName + FString(TEXT("\""));
+
+	// HKLM\SOFTWARE\Classes\.uproject
+	FRegistryRootedKey *RootExtensionKey = new FRegistryRootedKey(HKEY_LOCAL_MACHINE, TEXT("SOFTWARE\\Classes\\.uproject"));
+	RootExtensionKey->Key = new FRegistryKey();
+	RootExtensionKey->Key->SetValue(TEXT(""), TEXT("Unreal.ProjectFile"));
+	RootedKeys.AddRawItem(RootExtensionKey);
+
+	// HKLM\SOFTWARE\Classes\Unreal.ProjectFile
+	FRegistryRootedKey *RootFileTypeKey = new FRegistryRootedKey(HKEY_LOCAL_MACHINE, TEXT("SOFTWARE\\Classes\\Unreal.ProjectFile"));
+	RootFileTypeKey->Key = new FRegistryKey();
+	RootFileTypeKey->Key->SetValue(TEXT(""), TEXT("Unreal Engine Project File"));
+	RootFileTypeKey->Key->FindOrAddKey(L"DefaultIcon")->SetValue(TEXT(""), QuotedExecutableFileName);
+	RootedKeys.AddRawItem(RootFileTypeKey);
+
+	// HKLM\SOFTWARE\Classes\Unreal.ProjectFile\shell
+	FRegistryKey *ShellKey = RootFileTypeKey->Key->FindOrAddKey(TEXT("shell"));
+
+	// HKLM\SOFTWARE\Classes\Unreal.ProjectFile\shell\open
+	FRegistryKey *ShellOpenKey = ShellKey->FindOrAddKey(TEXT("open"));
+	ShellOpenKey->SetValue(L"", L"Open");
+	ShellOpenKey->FindOrAddKey(L"command")->SetValue(L"", QuotedExecutableFileName + FString(TEXT(" /editor \"%1\"")));
+
+	// HKLM\SOFTWARE\Classes\Unreal.ProjectFile\shell\run
+	FRegistryKey *ShellRunKey = ShellKey->FindOrAddKey(TEXT("run"));
+	ShellRunKey->SetValue(TEXT(""), TEXT("Launch game"));
+	ShellRunKey->SetValue(TEXT("Icon"), QuotedExecutableFileName);
+	ShellRunKey->FindOrAddKey(L"command")->SetValue(TEXT(""), QuotedExecutableFileName + FString(TEXT(" /game \"%1\"")));
+
+	// HKLM\SOFTWARE\Classes\Unreal.ProjectFile\shell\rungenproj
+	FRegistryKey *ShellProjectKey = ShellKey->FindOrAddKey(TEXT("rungenproj"));
+	ShellProjectKey->SetValue(L"", L"Generate Visual Studio project files");
+	ShellProjectKey->SetValue(L"Icon", QuotedExecutableFileName);
+	ShellProjectKey->FindOrAddKey(L"command")->SetValue(TEXT(""), QuotedExecutableFileName + FString(TEXT(" /projectfiles \"%1\"")));
+
+	// HKLM\SOFTWARE\Classes\Unreal.ProjectFile\shell\switchversion
+	FRegistryKey *ShellVersionKey = ShellKey->FindOrAddKey(TEXT("switchversion"));
+	ShellVersionKey->SetValue(TEXT(""), TEXT("Switch Unreal Engine version..."));
+	ShellVersionKey->SetValue(TEXT("Icon"), QuotedExecutableFileName);
+	ShellVersionKey->FindOrAddKey(L"command")->SetValue(TEXT(""), QuotedExecutableFileName + FString(TEXT(" /switchversion \"%1\"")));
+}
+
+int32 FDesktopPlatformWindows::GetShellIntegrationVersion(const FString &FileName)
+{
+	::DWORD VersionInfoSize = GetFileVersionInfoSize(*FileName, NULL);
+	if (VersionInfoSize != 0)
+	{
+		TArray<uint8> VersionInfo;
+		VersionInfo.AddUninitialized(VersionInfoSize);
+		if (GetFileVersionInfo(*FileName, NULL, VersionInfoSize, VersionInfo.GetData()))
+		{
+			TCHAR *ShellVersion;
+			::UINT ShellVersionLen;
+			if (VerQueryValue(VersionInfo.GetData(), TEXT("\\StringFileInfo\\040904b0\\ShellIntegrationVersion"), (LPVOID*)&ShellVersion, &ShellVersionLen))
+			{
+				TCHAR *ShellVersionEnd;
+				int32 Version = FCString::Strtoi(ShellVersion, &ShellVersionEnd, 10);
+				if(*ShellVersionEnd == 0)
+				{
+					return Version;
+				}
+			}
+		}
+	}
+	return 0;
+}
+
 #undef LOCTEXT_NAMESPACE
diff --git a/Engine/Source/Developer/DesktopPlatform/Private/Windows/DesktopPlatformWindows.h b/Engine/Source/Developer/DesktopPlatform/Private/Windows/DesktopPlatformWindows.h
--- a/Engine/Source/Developer/DesktopPlatform/Private/Windows/DesktopPlatformWindows.h
+++ b/Engine/Source/Developer/DesktopPlatform/Private/Windows/DesktopPlatformWindows.h
@@ -3,6 +3,7 @@
 #pragma once
 
 #include "../DesktopPlatformBase.h"
+#include "WindowsRegistry.h"
 
 class FDesktopPlatformWindows : public FDesktopPlatformBase
 {
@@ -16,8 +17,15 @@ public:
 
 	virtual void EnumerateEngineInstallations(TMap<FString, FString> &OutInstallations) OVERRIDE;
 
+	virtual bool IsSourceDistribution(const FString &RootDir) OVERRIDE;
+
+	virtual bool VerifyFileAssociations() OVERRIDE;
+	virtual bool UpdateFileAssociations() OVERRIDE;
+
 private:
 	bool FileDialogShared(bool bSave, const void* ParentWindowHandle, const FString& DialogTitle, const FString& DefaultPath, const FString& DefaultFile, const FString& FileTypes, uint32 Flags, TArray<FString>& OutFilenames);
+	void GetRequiredRegistrySettings(TIndirectArray<FRegistryRootedKey> &RootedKeys);
+	int32 GetShellIntegrationVersion(const FString &FileName);
 };
 
 typedef FDesktopPlatformWindows FDesktopPlatform;
diff --git a/Engine/Source/Developer/DesktopPlatform/Private/Windows/WindowsRegistry.cpp b/Engine/Source/Developer/DesktopPlatform/Private/Windows/WindowsRegistry.cpp
new file mode 100644
--- /dev/null
+++ b/Engine/Source/Developer/DesktopPlatform/Private/Windows/WindowsRegistry.cpp
@@ -0,0 +1,388 @@
+// Copyright 1998-2014 Epic Games, Inc. All Rights Reserved.
+
+#include "DesktopPlatformPrivatePCH.h"
+#include "WindowsRegistry.h"
+#include "AllowWindowsPlatformTypes.h"
+
+////////////////////////////////////////////////////////////////////
+
+FRegistryValue::FRegistryValue()
+{
+	Type = 0;
+}
+
+void FRegistryValue::Set(const FString &NewValue)
+{
+	Type = REG_SZ;
+	Data.Reset();
+	Data.Append((const uint8*)*NewValue, (NewValue.Len() + 1) * sizeof(TCHAR));
+}
+
+void FRegistryValue::Set(uint32 NewValue)
+{
+	Type = REG_DWORD;
+	Data.Reset();
+	Data.Append((const uint8*)&NewValue, sizeof(uint32));
+}
+
+bool FRegistryValue::Read(HKEY hKey, const FString &Name)
+{
+	// Read the size and type
+	DWORD ValueType;
+	DWORD ValueDataLength = 0;
+	if (RegQueryValueEx(hKey, *Name, NULL, &ValueType, NULL, &ValueDataLength) != ERROR_SUCCESS)
+	{
+		return false;
+	}
+	Type = (uint32)ValueType;
+
+	// Read the value data
+	Data.SetNum(ValueDataLength);
+	if (RegQueryValueEx(hKey, *Name, NULL, NULL, Data.GetData(), &ValueDataLength) != ERROR_SUCCESS || ValueDataLength != Data.Num())
+	{
+		return false;
+	}
+
+	// Otherwise success
+	return true;
+}
+
+bool FRegistryValue::Write(HKEY hKey, const FString &Name) const
+{
+	return (RegSetValueEx(hKey, *Name, 0, Type, Data.GetData(), (DWORD)Data.Num()) == ERROR_SUCCESS);
+}
+
+bool FRegistryValue::IsUpToDate(HKEY hKey, const FString &Name) const
+{
+	FRegistryValue Other;
+	return Other.Read(hKey, Name) && Other.Type == Type && Other.Data == Data;
+}
+
+////////////////////////////////////////////////////////////////////
+
+FRegistryKey::FRegistryKey()
+{
+}
+
+FRegistryKey::~FRegistryKey()
+{
+	for (TMap<FString, FRegistryKey*>::TIterator Iter(Keys); Iter; ++Iter)
+	{
+		delete Iter.Value();
+	}
+	for (TMap<FString, FRegistryValue*>::TIterator Iter(Values); Iter; ++Iter)
+	{
+		delete Iter.Value();
+	}
+}
+
+FRegistryKey *FRegistryKey::FindOrAddKey(const FString &Name)
+{
+	FRegistryKey *&Key = Keys.FindOrAdd(Name);
+	if (Key == NULL)
+	{
+		Key = new FRegistryKey();
+	}
+	return Key;
+}
+
+FRegistryValue *FRegistryKey::FindOrAddValue(const FString &Name)
+{
+	FRegistryValue *&Value = Values.FindOrAdd(Name);
+	if (Value == NULL)
+	{
+		Value = new FRegistryValue();
+	}
+	return Value;
+}
+
+void FRegistryKey::SetValue(const FString &Name, const FString &NewValue)
+{
+	FindOrAddValue(Name)->Set(NewValue);
+}
+
+void FRegistryKey::SetValue(const FString &Name, uint32 NewValue)
+{
+	FindOrAddValue(Name)->Set(NewValue);
+}
+
+bool FRegistryKey::Read(HKEY hKey)
+{
+	// Enumerate all the child key names
+	TArray<FString> KeyNames;
+	EnumerateRegistryKeys(hKey, KeyNames);
+
+	// Read all the child keys
+	for (size_t Idx = 0; Idx < KeyNames.Num(); Idx++)
+	{
+		FRegistryKey *Key = FindOrAddKey(KeyNames[Idx]);
+		if (!Key->Read(hKey, KeyNames[Idx]))
+		{
+			return false;
+		}
+	}
+
+	// Enumerate all the child values
+	TArray<FString> ValueNames;
+	EnumerateRegistryValues(hKey, ValueNames);
+
+	// Read all the values
+	for (size_t Idx = 0; Idx < ValueNames.Num(); Idx++)
+	{
+		FRegistryValue *Value = FindOrAddValue(ValueNames[Idx]);
+		if (!Value->Read(hKey, ValueNames[Idx]))
+		{
+			return false;
+		}
+	}
+
+	return true;
+}
+
+bool FRegistryKey::Read(HKEY hRootKey, const FString &Path)
+{
+	bool bRes = false;
+
+	HKEY hKey;
+	if (RegOpenKeyEx(hRootKey, *Path, 0, KEY_READ, &hKey) == ERROR_SUCCESS)
+	{
+		bRes = Read(hKey);
+		RegCloseKey(hKey);
+	}
+
+	return bRes;
+}
+
+bool FRegistryKey::Write(HKEY hKey, bool bRemoveDifferences) const
+{
+	// Remove all the differences from the current content
+	if (bRemoveDifferences)
+	{
+		// Get all the existing value names
+		TArray<FString> ValueNames;
+		EnumerateRegistryValues(hKey, ValueNames);
+
+		// Remove any values that don't exist any more
+		for (size_t Idx = 0; Idx < ValueNames.Num(); Idx++)
+		{
+			if (!Values.Contains(ValueNames[Idx]) && RegDeleteValue(hKey, *ValueNames[Idx]) != ERROR_SUCCESS)
+			{
+				return false;
+			}
+		}
+
+		// Get all the existing key names
+		TArray<FString> KeyNames;
+		EnumerateRegistryKeys(hKey, KeyNames);
+
+		// Remove any keys that don't exist any more
+		for (size_t Idx = 0; Idx < KeyNames.Num(); Idx++)
+		{
+			if (!Keys.Contains(KeyNames[Idx]) && RegDeleteTree(hKey, *KeyNames[Idx]) != ERROR_SUCCESS)
+			{
+				return false;
+			}
+		}
+	}
+
+	// Write all the child keys
+	for (TMap<FString, FRegistryKey*>::TConstIterator ChildKeyIter(Keys); ChildKeyIter; ++ChildKeyIter)
+	{
+		if (!ChildKeyIter.Value()->Write(hKey, ChildKeyIter.Key(), bRemoveDifferences))
+		{
+			return false;
+		}
+	}
+
+	// Write all the child values
+	for (TMap<FString, FRegistryValue*>::TConstIterator ChildValueIter(Values); ChildValueIter; ++ChildValueIter)
+	{
+		if (RegSetValueEx(hKey, *ChildValueIter.Key(), 0, ChildValueIter.Value()->Type, ChildValueIter.Value()->Data.GetData(), (DWORD)ChildValueIter.Value()->Data.Num()) != ERROR_SUCCESS)
+		{
+			return false;
+		}
+	}
+	return true;
+}
+
+bool FRegistryKey::Write(HKEY hRootKey, const FString &Path, bool bRemoveDifferences) const
+{
+	bool bRes = false;
+
+	HKEY hKey;
+	if (RegCreateKeyEx(hRootKey, *Path, 0, NULL, 0, KEY_READ | KEY_WRITE, NULL, &hKey, NULL) == ERROR_SUCCESS)
+	{
+		bRes = Write(hKey, bRemoveDifferences);
+		RegCloseKey(hKey);
+	}
+
+	return bRes;
+}
+
+bool FRegistryKey::IsUpToDate(HKEY hKey, bool bRemoveDifferences) const
+{
+	// Remove all the differences from the current content
+	if (bRemoveDifferences)
+	{
+		// Get all the existing value names
+		TArray<FString> ValueNames;
+		EnumerateRegistryValues(hKey, ValueNames);
+
+		// Check there aren't any extra values
+		for (size_t Idx = 0; Idx < ValueNames.Num(); Idx++)
+		{
+			if (!Values.Contains(ValueNames[Idx]))
+			{
+				return false;
+			}
+		}
+
+		// Get all the existing key names
+		TArray<FString> KeyNames;
+		EnumerateRegistryKeys(hKey, KeyNames);
+
+		// Check there aren't any extra keys
+		for (size_t Idx = 0; Idx < KeyNames.Num(); Idx++)
+		{
+			if (!Keys.Contains(KeyNames[Idx]))
+			{
+				return false;
+			}
+		}
+	}
+
+	// Write all the child keys
+	for (TMap<FString, FRegistryKey*>::TConstIterator Iter(Keys); Iter; ++Iter)
+	{
+		if (!Iter.Value()->IsUpToDate(hKey, Iter.Key(), bRemoveDifferences))
+		{
+			return false;
+		}
+	}
+
+	// Write all the child values
+	for (TMap<FString, FRegistryValue*>::TConstIterator Iter(Values); Iter; ++Iter)
+	{
+		if (!Iter.Value()->IsUpToDate(hKey, Iter.Key()))
+		{
+			return false;
+		}
+	}
+	return true;
+}
+
+bool FRegistryKey::IsUpToDate(HKEY hRootKey, const FString &Path, bool bRemoveDifferences) const
+{
+	bool bRes = false;
+
+	HKEY hKey;
+	if (RegOpenKeyEx(hRootKey, *Path, 0, KEY_READ, &hKey) == ERROR_SUCCESS)
+	{
+		bRes = IsUpToDate(hKey, bRemoveDifferences);
+		RegCloseKey(hKey);
+	}
+
+	return bRes;
+}
+
+////////////////////////////////////////////////////////////////////
+
+FRegistryRootedKey::FRegistryRootedKey(HKEY hInRootKey, const FString &InPath)
+{
+	hRootKey = hInRootKey;
+	Path = InPath;
+}
+
+bool FRegistryRootedKey::Exists() const
+{
+	bool bRes = false;
+
+	HKEY hKey;
+	if (RegOpenKeyEx(hRootKey, *Path, 0, KEY_READ, &hKey) == ERROR_SUCCESS)
+	{
+		RegCloseKey(hKey);
+		bRes = true;
+	}
+
+	return bRes;
+}
+
+bool FRegistryRootedKey::Write(bool bRemoveDifferences) const
+{
+	bool bRes = false;
+	if (Key)
+	{
+		bRes = Key->Write(hRootKey, Path, bRemoveDifferences);
+	}
+	else
+	{
+		bRes = (!Exists() || RegDeleteKeyEx(hRootKey, *Path, 0, 0) == ERROR_SUCCESS);
+	}
+	return bRes;
+}
+
+bool FRegistryRootedKey::IsUpToDate(bool bRemoveDifferences) const
+{
+	bool bRes = false;
+	if (Key)
+	{
+		bRes = Key->IsUpToDate(hRootKey, Path, bRemoveDifferences);
+	}
+	else
+	{
+		bRes = !Exists();
+	}
+	return bRes;
+}
+
+////////////////////////////////////////////////////////////////////
+
+bool EnumerateRegistryKeys(HKEY hKey, TArray<FString> &OutNames)
+{
+	for (DWORD Index = 0;; Index++)
+	{
+		// Query the next key name
+		TCHAR KeyName[256];
+		DWORD KeyNameLength = sizeof(KeyName) / sizeof(KeyName[0]);
+
+		LONG Result = RegEnumKeyEx(hKey, Index, KeyName, &KeyNameLength, NULL, NULL, NULL, NULL);
+		if (Result == ERROR_NO_MORE_ITEMS)
+		{
+			break;
+		}
+		else if (Result != ERROR_SUCCESS)
+		{
+			return false;
+		}
+
+		// Add it to the array
+		OutNames.Add(KeyName);
+	}
+	return true;
+}
+
+bool EnumerateRegistryValues(HKEY hKey, TArray<FString> &OutNames)
+{
+	for (DWORD Index = 0;; Index++)
+	{
+		// Query the value
+		wchar_t ValueName[256];
+		DWORD ValueNameLength = sizeof(ValueName) / sizeof(ValueName[0]);
+
+		LONG Result = RegEnumValue(hKey, Index, ValueName, &ValueNameLength, NULL, NULL, NULL, NULL);
+		if (Result == ERROR_NO_MORE_ITEMS)
+		{
+			break;
+		}
+		else if (Result != ERROR_SUCCESS)
+		{
+			return false;
+		}
+
+		// Add it to the array
+		OutNames.Add(ValueName);
+	}
+	return true;
+}
+
+#include "HideWindowsPlatformTypes.h"
diff --git a/Engine/Source/Developer/DesktopPlatform/Private/Windows/WindowsRegistry.h b/Engine/Source/Developer/DesktopPlatform/Private/Windows/WindowsRegistry.h
new file mode 100644
--- /dev/null
+++ b/Engine/Source/Developer/DesktopPlatform/Private/Windows/WindowsRegistry.h
@@ -0,0 +1,61 @@
+// Copyright 1998-2014 Epic Games, Inc. All Rights Reserved.
+
+#pragma once
+
+struct FRegistryValue
+{
+	uint32 Type;
+	TArray<uint8> Data;
+
+	FRegistryValue();
+
+	void Set(const FString &NewValue);
+	void Set(uint32 NewValue);
+
+	bool Read(HKEY hKey, const FString &Name);
+	bool Write(HKEY hKey, const FString &Name) const;
+	bool IsUpToDate(HKEY hKey, const FString &Name) const;
+};
+
+struct FRegistryKey : FNoncopyable
+{
+	TMap<FString, FRegistryKey*> Keys;
+	TMap<FString, FRegistryValue*> Values;
+
+	FRegistryKey();
+	~FRegistryKey();
+
+	FRegistryKey *FindOrAddKey(const FString &Name);
+	FRegistryValue *FindOrAddValue(const FString &Name);
+
+	void SetValue(const FString &Name, const FString &Value);
+	void SetValue(const FString &Name, uint32 Value);
+
+	bool Read(HKEY hKey);
+	bool Read(HKEY hKey, const FString &Path);
+
+	bool Write(HKEY hKey, bool bRemoveDifferences) const;
+	bool Write(HKEY hKey, const FString &Path, bool bRemoveDifferences) const;
+
+	bool IsUpToDate(HKEY hKey, bool bRemoveDifferences) const;
+	bool IsUpToDate(HKEY hKey, const FString &Path, bool bRemoveDifferences) const;
+};
+
+struct FRegistryRootedKey
+{
+	HKEY hRootKey;
+	FString Path;
+	TScopedPointer<FRegistryKey> Key;
+
+	FRegistryRootedKey(HKEY hInKeyRoot, const FString &InPath);
+
+	bool Exists() const;
+
+	bool Read();
+	bool Write(bool bRemoveDifferences) const;
+	bool IsUpToDate(bool bRemoveDifferences) const;
+};
+
+bool EnumerateRegistryKeys(HKEY hKey, TArray<FString> &OutNames);
+bool EnumerateRegistryValues(HKEY hKey, TArray<FString> &OutNames);
+
diff --git a/Engine/Source/Developer/DesktopPlatform/Public/IDesktopPlatform.h b/Engine/Source/Developer/DesktopPlatform/Public/IDesktopPlatform.h
--- a/Engine/Source/Developer/DesktopPlatform/Public/IDesktopPlatform.h
+++ b/Engine/Source/Developer/DesktopPlatform/Public/IDesktopPlatform.h
@@ -132,6 +132,13 @@ public:
 	*/
 	virtual void EnumerateEngineInstallations(TMap<FString, FString> &OutInstallations) = 0;
 
+	/**
+	* Enumerates all the registered binary engine installations.
+	*
+	* @param	OutInstallations	Array which is filled in with identifier/root-directory pairs for all known binary installations.
+	*/
+	virtual void EnumerateLauncherEngineInstallations(TMap<FString, FString> &OutInstallations) = 0;
+
 	/**
 	* Returns the identifier for the engine with the given root directory.
 	*
@@ -147,4 +154,82 @@ public:
 	* @param	OutIdentifier		Identifier used to refer to this installation.
 	*/
 	virtual bool GetEngineIdentifierFromRootDir(const FString &RootDir, FString &OutIdentifier) = 0;
+
+	/**
+	* Gets the identifier for the default engine. This will be the newest installed engine.
+	*
+	* @param	OutIdentifier	String to hold to the default engine's identifier
+	* @return	true if OutIdentifier is valid.
+	*/
+	virtual bool GetDefaultEngineIdentifier(FString &OutIdentifier) = 0;
+
+	/**
+	* Compares two identifiers and checks whether the first is preferred to the second.
+	*
+	* @param	Identifier		First identifier
+	* @param	OtherIdentifier	Second identifier
+	* @return	true if Identifier is preferred over OtherIdentifier
+	*/
+	virtual bool IsPreferredEngineIdentifier(const FString &Identifier, const FString &OtherIdentifier) = 0;
+
+	/**
+	* Gets the root directory for the default engine installation.
+	*
+	* @param	OutRootDir	String to hold to the default engine's root directory
+	* @return	true if OutRootDir is valid
+	*/
+	virtual bool GetDefaultEngineRootDir(FString &OutRootDir) = 0;
+
+	/**
+	* Tests whether an engine installation is a source distribution.
+	*
+	* @return	true if the engine contains source.
+	*/
+	virtual bool IsSourceDistribution(const FString &RootDir) = 0;
+
+	/**
+	* Tests whether an engine installation is a perforce build.
+	*
+	* @return	true if the engine is a perforce build.
+	*/
+	virtual bool IsPerforceBuild(const FString &RootDir) = 0;
+
+	/**
+	* Tests whether a root directory is a valid Unreal Engine installation
+	*
+	* @return	true if the engine is a valid installation
+	*/
+	virtual bool IsValidRootDirectory(const FString &RootDir) = 0;
+
+	/**
+	* Checks that the current file associations are correct.
+	*
+	* @return	true if file associations are up to date.
+	*/
+	virtual bool VerifyFileAssociations() = 0;
+
+	/**
+	* Updates file associations.
+	*
+	* @return	true if file associations were successfully updated.
+	*/
+	virtual bool UpdateFileAssociations() = 0;
+
+	/**
+	* Sets the engine association for a project.
+	*
+	* @param ProjectFileName	Filename of the project to update
+	* @param Identifier			Identifier of the engine to associate it with
+	* @return true if the project was successfully updated
+	*/
+	virtual bool SetEngineIdentifierForProject(const FString &ProjectFileName, const FString &Identifier) = 0;
+
+	/**
+	* Gets the engine association for a project.
+	*
+	* @param ProjectFileName	Filename of the project to update
+	* @param OutIdentifier		Identifier of the engine to associate it with
+	* @return true if OutIdentifier is set to the project's engine association
+	*/
+	virtual bool GetEngineIdentifierForProject(const FString &ProjectFileName, FString &OutIdentifier) = 0;
 };
diff --git a/Engine/Source/Programs/UnrealVersionSelector/Private/FileAssociation.cpp b/Engine/Source/Programs/UnrealVersionSelector/Private/FileAssociation.cpp
deleted file mode 100644
diff --git a/Engine/Source/Programs/UnrealVersionSelector/Private/FileAssociation.h b/Engine/Source/Programs/UnrealVersionSelector/Private/FileAssociation.h
deleted file mode 100644
diff --git a/Engine/Source/Programs/UnrealVersionSelector/Private/GenericPlatform/GenericPlatformInstallation.cpp b/Engine/Source/Programs/UnrealVersionSelector/Private/GenericPlatform/GenericPlatformInstallation.cpp
--- a/Engine/Source/Programs/UnrealVersionSelector/Private/GenericPlatform/GenericPlatformInstallation.cpp
+++ b/Engine/Source/Programs/UnrealVersionSelector/Private/GenericPlatform/GenericPlatformInstallation.cpp
@@ -5,70 +5,10 @@
 #include "../PlatformInstallation.h"
 #include "Runtime/Launch/Resources/Version.h"
 
-bool FGenericPlatformInstallation::GetLauncherVersionSelector(FString &OutFileName)
-{
-	FConfigFile ConfigFile;
-	ConfigFile.Read(GetConfigFilePath());
-
-	FString InstallationPath;
-	if (!ConfigFile.GetString(TEXT("Launcher"), TEXT("InstalledPath"), InstallationPath))
-	{
-		return false;
-	}
-
-	FString LauncherFileName = InstallationPath / TEXT("Engine/Binaries/Win64/UnrealVersionSelector.exe");
-	if (IFileManager::Get().FileSize(*LauncherFileName) < 0)
-	{
-		return false;
-	}
-
-	return true;
-}
-
 bool FGenericPlatformInstallation::RegisterEngineInstallation(const FString &Id, const FString &RootDirName)
 {
-	FString NormalizedRootDirName = RootDirName;
-	FPaths::NormalizeDirectoryName(NormalizedRootDirName);
-
-	FConfigFile ConfigFile;
-	ConfigFile.Read(GetConfigFilePath());
-
-	FConfigSection &Section = ConfigFile.FindOrAdd(TEXT("Installations"));
-	Section.AddUnique(*Id, NormalizedRootDirName);
-	ConfigFile.Dirty = true;
-	return ConfigFile.Write(GetConfigFilePath());
-}
-
-bool FGenericPlatformInstallation::UnregisterEngineInstallation(const FString &Id)
-{
-	FConfigFile ConfigFile;
-	ConfigFile.Read(GetConfigFilePath());
-
-	FConfigSection &Section = ConfigFile.FindOrAdd(TEXT("Installations"));
-	Section.Remove(*Id);
-	ConfigFile.Dirty = true;
-	return ConfigFile.Write(GetConfigFilePath());
-}
-
-bool FGenericPlatformInstallation::IsSourceDistribution(const FString &EngineRootDir)
-{
-	// Check for the existence of a SourceBuild.txt file
-	FString SourceBuildPath = EngineRootDir / TEXT("Engine\\Build\\SourceDistribution.txt");
-	return (IFileManager::Get().FileSize(*SourceBuildPath) >= 0);
-}
-
-bool FGenericPlatformInstallation::IsPerforceBuild(const FString &EngineRootDir)
-{
-	// Check for the existence of a SourceBuild.txt file
-	FString PerforceBuildPath = EngineRootDir / TEXT("Engine\\Build\\PerforceBuild.txt");
-	return (IFileManager::Get().FileSize(*PerforceBuildPath) >= 0);
-}
-
-bool FGenericPlatformInstallation::IsValidEngineRootDir(const FString &RootDir)
-{
-	// Check if the engine binaries directory exists
-	FString EngineBinariesDirName = RootDir / TEXT("Engine/Binaries");
-	return IFileManager::Get().DirectoryExists(*EngineBinariesDirName);
+	// No default implementation
+	return true;
 }
 
 bool FGenericPlatformInstallation::NormalizeEngineRootDir(FString &RootDir)
@@ -78,7 +18,7 @@ bool FGenericPlatformInstallation::NormalizeEngineRootDir(FString &RootDir)
 	FPaths::NormalizeDirectoryName(NormalizedRootDir);
 
 	// Check if it's valid
-	if (IsValidEngineRootDir(NormalizedRootDir))
+	if (FDesktopPlatformModule::Get()->IsValidRootDirectory(NormalizedRootDir))
 	{
 		RootDir = NormalizedRootDir;
 		return true;
@@ -94,7 +34,7 @@ bool FGenericPlatformInstallation::NormalizeEngineRootDir(FString &RootDir)
 	}
 
 	// Check if the engine binaries directory exists
-	if (IsValidEngineRootDir(NormalizedRootDir))
+	if (FDesktopPlatformModule::Get()->IsValidRootDirectory(NormalizedRootDir))
 	{
 		RootDir = NormalizedRootDir;
 		return true;
@@ -115,20 +55,8 @@ bool FGenericPlatformInstallation::GenerateProjectFiles(const FString &RootDirNa
 	return false;
 }
 
-bool FGenericPlatformInstallation::VerifyShellIntegration()
-{
-	// No shell integration by default
-	return true;
-}
-
-bool FGenericPlatformInstallation::UpdateShellIntegration()
-{
-	// No shell integration by default
-	return true;
-}
-
-FString FGenericPlatformInstallation::GetConfigFilePath()
+bool FGenericPlatformInstallation::SelectEngineInstallation(FString &Identifier)
 {
-	return FString(FPlatformProcess::ApplicationSettingsDir()) / FString(EPIC_PRODUCT_IDENTIFIER) / FString(TEXT("Install.ini"));
+	// No default implementation
+	return false;
 }
-
diff --git a/Engine/Source/Programs/UnrealVersionSelector/Private/GenericPlatform/GenericPlatformInstallation.h b/Engine/Source/Programs/UnrealVersionSelector/Private/GenericPlatform/GenericPlatformInstallation.h
--- a/Engine/Source/Programs/UnrealVersionSelector/Private/GenericPlatform/GenericPlatformInstallation.h
+++ b/Engine/Source/Programs/UnrealVersionSelector/Private/GenericPlatform/GenericPlatformInstallation.h
@@ -6,24 +6,9 @@
 
 struct FGenericPlatformInstallation
 {
-	/// Get the path to the installed build of the version selector
-	static bool GetLauncherVersionSelector(FString &OutFileName);
-
 	// Register an engine installation
 	static bool RegisterEngineInstallation(const FString &Id, const FString &RootDirName);
 
-	// Unregister an engine installation
-	static bool UnregisterEngineInstallation(const FString &Id);
-
-	// Test whether a given engine root directory is a source build
-	static bool IsSourceDistribution(const FString &EngineRootDir);
-
-	// Tests whether an engine directory contains a Perforce build
-	static bool IsPerforceBuild(const FString &EngineRootDir);
-
-	// Validate and normalize an engine root directory name
-	static bool IsValidEngineRootDir(const FString &RootDir);
-
 	// Validate and normalize an engine root directory name
 	static bool NormalizeEngineRootDir(FString &RootDir);
 
@@ -33,16 +18,6 @@ struct FGenericPlatformInstallation
 	// Generate project files by running UBT
 	static bool GenerateProjectFiles(const FString &RootDirName, const FString &Arguments);
 
-	// Check whether the current settings are up to date
-	static bool VerifyShellIntegration();
-
-	// Update the installation settings
-	static bool UpdateShellIntegration();
-
-private:
-	// Returns the config file for this 
-	static FString GetConfigFilePath();
-
-	// Returns the config file for this 
-	static FConfigFile *GetConfigFile();
+	// Select an engine installation
+	static bool SelectEngineInstallation(FString &Identifier);
 };
diff --git a/Engine/Source/Programs/UnrealVersionSelector/Private/PlatformInstallation.cpp b/Engine/Source/Programs/UnrealVersionSelector/Private/PlatformInstallation.cpp
--- a/Engine/Source/Programs/UnrealVersionSelector/Private/PlatformInstallation.cpp
+++ b/Engine/Source/Programs/UnrealVersionSelector/Private/PlatformInstallation.cpp
@@ -2,53 +2,3 @@
 #include "PlatformInstallation.h"
 #include "DesktopPlatformModule.h"
 
-int32 ParseReleaseVersion(const FString &Version)
-{
-	TCHAR *End;
-
-	uint64 Major = FCString::Strtoui64(*Version, &End, 10);
-	if (Major >= MAX_int16 || *(End++) != '.')
-	{
-		return INDEX_NONE;
-	}
-
-	uint64 Minor = FCString::Strtoui64(End, &End, 10);
-	if (Minor >= MAX_int16 || *End != 0)
-	{
-		return INDEX_NONE;
-	}
-
-	return (Major << 16) + Minor;
-}
-
-bool GetDefaultEngineId(FString &OutId)
-{
-	TMap<FString, FString> Installations;
-	FDesktopPlatformModule::Get()->EnumerateEngineInstallations(Installations);
-
-	bool bRes = false;
-	if (Installations.Num() > 0)
-	{
-		// Default to the first install
-		OutId = TMap<FString, FString>::TConstIterator(Installations).Key();
-
-		// Try to find the highest versioned install
-		int32 BestVersion = INDEX_NONE;
-		for (TMap<FString, FString>::TConstIterator Iter(Installations); Iter; ++Iter)
-		{
-			int32 Version = ParseReleaseVersion(Iter.Key());
-			if (Version > BestVersion)
-			{
-				OutId = Iter.Key();
-				BestVersion = Version;
-			}
-		}
-	}
-	return bRes;
-}
-
-bool GetDefaultEngineRootDir(FString &OutDirName)
-{
-	FString Id;
-	return GetDefaultEngineId(Id) && FDesktopPlatformModule::Get()->GetEngineRootDirFromIdentifier(Id, OutDirName);
-}
diff --git a/Engine/Source/Programs/UnrealVersionSelector/Private/PlatformInstallation.h b/Engine/Source/Programs/UnrealVersionSelector/Private/PlatformInstallation.h
--- a/Engine/Source/Programs/UnrealVersionSelector/Private/PlatformInstallation.h
+++ b/Engine/Source/Programs/UnrealVersionSelector/Private/PlatformInstallation.h
@@ -11,10 +11,3 @@
 	#include "GenericPlatform/GenericPlatformInstallation.h"
 	typedef FGenericPlatformInstallation FPlatformInstallation;
 #endif
-
-// Utility functions
-int32 ParseReleaseVersion(const FString &Version);
-
-bool GetDefaultEngineId(FString &OutId);
-bool GetDefaultEngineRootDir(FString &OutId);
-
diff --git a/Engine/Source/Programs/UnrealVersionSelector/Private/UnrealVersionSelector.cpp b/Engine/Source/Programs/UnrealVersionSelector/Private/UnrealVersionSelector.cpp
--- a/Engine/Source/Programs/UnrealVersionSelector/Private/UnrealVersionSelector.cpp
+++ b/Engine/Source/Programs/UnrealVersionSelector/Private/UnrealVersionSelector.cpp
@@ -2,175 +2,88 @@
 #include "RequiredProgramMainCPPInclude.h"
 #include "DesktopPlatformModule.h"
 #include "PlatformInstallation.h"
-#include "FileAssociation.h"
 
 IMPLEMENT_APPLICATION(UnrealVersionSelector, "UnrealVersionSelector")
 
-FString FullExecutableName()
-{
-	return FString(FPlatformProcess::BaseDir()) / FString(FPlatformProcess::ExecutableName(false));
-}
-
 bool RegisterCurrentEngineDirectory()
 {
-	// Get the current engine directory.
-	FString EngineRootDir = FPlatformMisc::RootDir();
-
-	// Get the canonical version selector path
-	FString VersionSelectorFileName;
-	if (!FPlatformInstallation::GetLauncherVersionSelector(VersionSelectorFileName))
-	{
-		VersionSelectorFileName = FString(FPlatformProcess::BaseDir()) / FString(FPlatformProcess::ExecutableName(false));
-	}
-
-	// Launch the installed version selector as elevated to register this directory
-	FString Arguments = FString::Printf(TEXT("/register \"%s\""), *EngineRootDir);
-	int32 ExitCode;
-	if (!FPlatformProcess::ExecElevatedProcess(*VersionSelectorFileName, *Arguments, &ExitCode) || ExitCode != 0)
+	// Prompt for registering this directory
+	if(FPlatformMisc::MessageBoxExt(EAppMsgType::YesNo, TEXT("Configure this directory as an Unreal Engine installation?"), TEXT("Question")) != EAppReturnType::Yes)
 	{
 		return false;
 	}
 
-	// Notify the user that everything is awesome.
-	FPlatformMisc::MessageBoxExt(EAppMsgType::Ok, TEXT("Registration successful."), TEXT("Success"));
-	return true;
-}
-
-void EnumerateInvalidInstallations(TArray<FString> &OutIds)
-{
-	// Enumerate all the engine installations
-	TMap<FString, FString> Installations;
-	FDesktopPlatformModule::Get()->EnumerateEngineInstallations(Installations);
+	// Get the current engine directory.
+	FString EngineRootDir = FPlatformProcess::BaseDir();
+	FPlatformInstallation::NormalizeEngineRootDir(EngineRootDir);
 
-	// Check each one of them has a valid engine path
-	for (TMap<FString, FString>::TConstIterator Iter(Installations); Iter; ++Iter)
+	// Get any existing tag name or register a new one
+	FString Identifier;
+	if (!FDesktopPlatformModule::Get()->GetEngineIdentifierFromRootDir(EngineRootDir, Identifier))
 	{
-		const FString &RootDir = Iter.Value();
-		if (!FPlatformInstallation::IsValidEngineRootDir(RootDir))
+		if(!FPlatformInstallation::RegisterEngineInstallation(EngineRootDir, Identifier))
 		{
-			OutIds.Add(Iter.Key());
+			FPlatformMisc::MessageBoxExt(EAppMsgType::Ok, TEXT("Couldn't add engine installation."), TEXT("Error"));
+			return false;
 		}
 	}
-}
 
-bool VerifySettings()
-{
-	// Enumerate all the invalid installations
-	TArray<FString> InvalidInstallationIds;
-	EnumerateInvalidInstallations(InvalidInstallationIds);
-
-	// If the current settings are valid, we don't need to do anything
-	if (InvalidInstallationIds.Num() == 0 && FPlatformInstallation::VerifyShellIntegration())
+	// If the launcher isn't installed, set up the file associations
+	if(!FDesktopPlatformModule::Get()->VerifyFileAssociations())
 	{
-		return true;
-	}
-
-	// Otherwise relaunch as administrator
-	return FPlatformProcess::ExecElevatedProcess(*FullExecutableName(), TEXT("/update"), NULL);
-}
+		// Relaunch as administrator
+		FString ExecutableFileName = FString(FPlatformProcess::BaseDir()) / FString(FPlatformProcess::ExecutableName(false));
 
-bool UpdateSettings()
-{
-	// Enumerate all the invalid installations
-	TArray<FString> InvalidInstallationIds;
-	EnumerateInvalidInstallations(InvalidInstallationIds);
-
-	// Remove those installations
-	for (int32 Idx = 0; Idx < InvalidInstallationIds.Num(); Idx++)
-	{
-		if (!FPlatformInstallation::UnregisterEngineInstallation(InvalidInstallationIds[Idx]))
+		int32 ExitCode;
+		if (!FPlatformProcess::ExecElevatedProcess(*ExecutableFileName, TEXT("/fileassociations"), &ExitCode) || ExitCode != 0)
 		{
-			FPlatformMisc::MessageBoxExt(EAppMsgType::Ok, TEXT("Couldn't update installation settings."), TEXT("Error"));
 			return false;
 		}
 	}
 
-	// Update everything
-	if (!FPlatformInstallation::UpdateShellIntegration())
-	{
-		FPlatformMisc::MessageBoxExt(EAppMsgType::Ok, TEXT("Couldn't update installation settings."), TEXT("Error"));
-		return false;
-	}
+	// Notify the user that everything is awesome.
+	FPlatformMisc::MessageBoxExt(EAppMsgType::Ok, TEXT("Registration successful."), TEXT("Success"));
 	return true;
 }
 
-bool RegisterEngineDirectory(const FString &EngineId, const FString &EngineDir)
+bool UpdateFileAssociations()
 {
-	// Add the installation by name
-	if (!FPlatformInstallation::RegisterEngineInstallation(EngineId, EngineDir))
+	// Update everything
+	if (!FDesktopPlatformModule::Get()->UpdateFileAssociations())
 	{
-		FPlatformMisc::MessageBoxExt(EAppMsgType::Ok, TEXT("Couldn't add engine installation."), TEXT("Error"));
+		FPlatformMisc::MessageBoxExt(EAppMsgType::Ok, TEXT("Couldn't update file associations."), TEXT("Error"));
 		return false;
 	}
-
-	// Add an installation, then call update
-	return UpdateSettings();
+	return true;
 }
 
-bool RegisterEngineDirectory(const FString &EngineDir)
+bool SwitchVersion(const FString &ProjectFileName)
 {
-	// Get any existing tag name or register a new one
-	FString EngineId;
-	if (!FDesktopPlatformModule::Get()->GetEngineIdentifierFromRootDir(EngineDir, EngineId))
+	// Get the current identifier
+	FString Identifier;
+	FDesktopPlatformModule::Get()->GetEngineIdentifierForProject(ProjectFileName, Identifier);
+
+	// Select the new association
+	if(!FPlatformInstallation::SelectEngineInstallation(Identifier))
 	{
-		EngineId = FGuid::NewGuid().ToString(EGuidFormats::DigitsWithHyphensInBraces);
+		return false;
 	}
 
-	// Add an installation, then call update
-	return RegisterEngineDirectory(EngineId, EngineDir);
-}
-
-bool SelectEngineAssociation(const FString &ProjectFileName, const FString &EngineId)
-{
-	// Just write the ID into the file
-	if (!SetEngineIdForProject(ProjectFileName, EngineId))
+	// Update the project file
+	if (!FDesktopPlatformModule::Get()->SetEngineIdentifierForProject(ProjectFileName, Identifier))
 	{
 		FPlatformMisc::MessageBoxExt(EAppMsgType::Ok, TEXT("Couldn't set association for project. Check the file is writeable."), TEXT("Error"));
 		return false;
 	}
 
 	// Notify the user that it's been changed
-	FPlatformMisc::MessageBoxExt(EAppMsgType::Ok, TEXT("Updated project file."), TEXT("Success"));
 	return true;
 }
 
-bool BrowseForEngineAssociation(const FString &ProjectFileName)
+bool GetEngineRootDirForProject(const FString &ProjectFileName, FString &OutRootDir)
 {
-	// Get the currently bound engine directory for the project
-	FString EngineRootDir;
-	GetEngineRootDirForProject(ProjectFileName, EngineRootDir);
-
-	// Browse for a new directory
-	FString NewEngineRootDir;
-	if (!FDesktopPlatformModule::Get()->OpenDirectoryDialog(NULL, TEXT("Select the Unreal Engine installation to use for this project"), EngineRootDir, NewEngineRootDir))
-	{
-		return false;
-	}
-
-	// Check it's a valid directory
-	if (!FPlatformInstallation::NormalizeEngineRootDir(NewEngineRootDir))
-	{
-		FPlatformMisc::MessageBoxExt(EAppMsgType::Ok, TEXT("The selected directory is not a valid engine installation."), TEXT("Error"));
-		return false;
-	}
-
-	// Check that it's a registered engine directory
-	FString EngineId;
-	if (!FDesktopPlatformModule::Get()->GetEngineIdentifierFromRootDir(NewEngineRootDir, EngineId))
-	{
-		int32 ExitCode;
-		if (!FPlatformProcess::ExecElevatedProcess(*FullExecutableName(), *FString::Printf(TEXT("/register \"%s\""), *NewEngineRootDir), &ExitCode) || ExitCode != 0)
-		{
-			return false;
-		}
-		if (!FDesktopPlatformModule::Get()->GetEngineIdentifierFromRootDir(NewEngineRootDir, EngineId))
-		{
-			return false;
-		}
-	}
-
-	// Set the file association
-	return SelectEngineAssociation(ProjectFileName, EngineId);
+	FString Identifier;
+	return FDesktopPlatformModule::Get()->GetEngineIdentifierForProject(ProjectFileName, Identifier) && FDesktopPlatformModule::Get()->GetEngineRootDirFromIdentifier(Identifier, OutRootDir);
 }
 
 bool GetValidatedEngineRootDir(const FString &ProjectFileName, FString &OutRootDir)
@@ -179,7 +92,7 @@ bool GetValidatedEngineRootDir(const FString &ProjectFileName, FString &OutRootD
 	if (!GetEngineRootDirForProject(ProjectFileName, OutRootDir))
 	{
 		// Try to set an association
-		if (!BrowseForEngineAssociation(ProjectFileName))
+		if(!SwitchVersion(ProjectFileName))
 		{
 			return false;
 		}
@@ -224,7 +137,7 @@ bool GenerateProjectFiles(const FString &ProjectFileName)
 
 	// Build the argument list
 	FString Arguments = TEXT("-game");
-	if (FPlatformInstallation::IsSourceDistribution(RootDir))
+	if (FDesktopPlatformModule::Get()->IsSourceDistribution(RootDir))
 	{
 		Arguments += TEXT(" -engine");
 	}
@@ -247,35 +160,15 @@ int Main(const TArray<FString> &Arguments)
 		// Add the current directory to the list of installations
 		bRes = RegisterCurrentEngineDirectory();
 	}
-	else if (Arguments.Num() == 1 && Arguments[0] == TEXT("/verify"))
-	{
-		// Verify all the settings are correct.
-		bRes = VerifySettings();
-	}
-	else if (Arguments.Num() == 1 && Arguments[0] == TEXT("/update"))
+	else if (Arguments.Num() == 1 && Arguments[0] == TEXT("/fileassociations"))
 	{
 		// Update all the settings.
-		bRes = UpdateSettings();
-	}
-	else if (Arguments.Num() == 2 && Arguments[0] == TEXT("/register"))
-	{
-		// Register the given directory as an engine directory
-		bRes = RegisterEngineDirectory(Arguments[1]);
+		bRes = UpdateFileAssociations();
 	}
-	else if (Arguments.Num() == 3 && Arguments[0] == TEXT("/register"))
-	{
-		// Register an id corresponding to a given directory
-		bRes = RegisterEngineDirectory(Arguments[1], Arguments[2]);
-	}
-	else if (Arguments.Num() == 3 && Arguments[0] == TEXT("/associate"))
+	else if (Arguments.Num() == 2 && Arguments[0] == TEXT("/switchversion"))
 	{
 		// Associate with an engine label
-		bRes = SelectEngineAssociation(Arguments[1], Arguments[2]);
-	}
-	else if (Arguments.Num() == 2 && Arguments[0] == TEXT("/browse"))
-	{
-		// Associate the project with a user selected engine
-		bRes = BrowseForEngineAssociation(Arguments[1]);
+		bRes = SwitchVersion(Arguments[1]);
 	}
 	else if (Arguments.Num() == 2 && Arguments[0] == TEXT("/editor"))
 	{
diff --git a/Engine/Source/Programs/UnrealVersionSelector/Private/Windows/WindowsPlatformInstallation.cpp b/Engine/Source/Programs/UnrealVersionSelector/Private/Windows/WindowsPlatformInstallation.cpp
--- a/Engine/Source/Programs/UnrealVersionSelector/Private/Windows/WindowsPlatformInstallation.cpp
+++ b/Engine/Source/Programs/UnrealVersionSelector/Private/Windows/WindowsPlatformInstallation.cpp
@@ -2,34 +2,27 @@
 
 #include "../UnrealVersionSelector.h"
 #include "WindowsPlatformInstallation.h"
-#include "WindowsRegistry.h"
 #include "Runtime/Core/Public/Misc/EngineVersion.h"
 #include "AllowWindowsPlatformTypes.h"
+#include "Resources/Resource.h"
+#include <CommCtrl.h>
 #include <Shlwapi.h>
 #include <Shellapi.h>
 
+const TCHAR *InstallationsSubKey = TEXT("SOFTWARE\\Epic Games\\Unreal Engine\\Builds");
+
 struct FEngineLabelSortPredicate
 {
 	bool operator()(const FString &A, const FString &B) const
 	{
-		// Return true if A < B
-		int32 VerA = ParseReleaseVersion(A);
-		int32 VerB = ParseReleaseVersion(B);
-		if (VerA == VerB)
-		{
-			return StrCmpLogicalW(*A, *B) < 0;
-		}
-		else
-		{
-			return VerA > VerB;
-		}
+		return FDesktopPlatformModule::Get()->IsPreferredEngineIdentifier(A, B);
 	}
 };
 
 FString GetInstallationDescription(const FString &Id, const FString &RootDir)
 {
 	// Official release versions just have a version number
-	if (Id.Len() > 0 && Id[0] != '{')
+	if (Id.Len() > 0 && FChar::IsDigit(Id[0]))
 	{
 		return Id;
 	}
@@ -39,174 +32,141 @@ FString GetInstallationDescription(const FString &Id, const FString &RootDir)
 	FPaths::MakePlatformFilename(PlatformRootDir);
 
 	// Perforce build
-	if (FPlatformInstallation::IsSourceDistribution(RootDir))
+	if (FDesktopPlatformModule::Get()->IsSourceDistribution(RootDir))
 	{
-		if (FPlatformInstallation::IsPerforceBuild(RootDir))
-		{
-			return FString::Printf(TEXT("Perforce build at %s"), *PlatformRootDir);
-		}
-		else
-		{
-			return FString::Printf(TEXT("GitHub build at %s"), *PlatformRootDir);
-		}
+		return FString::Printf(TEXT("Source build at %s"), *PlatformRootDir);
 	}
 	else
 	{
-		return FString::Printf(TEXT("%s"), *PlatformRootDir);
+		return FString::Printf(TEXT("Binary build at %s"), *PlatformRootDir);
 	}
 }
 
-void GetRequiredRegistrySettings(TIndirectArray<FRegistryRootedKey> &RootedKeys)
+struct FSelectBuildDialog
 {
-	// Get the path to the executable
-	FString ExecutableFileName = FString(FPlatformProcess::BaseDir()) / FString(FPlatformProcess::ExecutableName(false));
-	FPaths::MakePlatformFilename(ExecutableFileName);
-	FString QuotedExecutableFileName = FString(TEXT("\"")) + ExecutableFileName + FString(TEXT("\""));
-
-	// HKLM\SOFTWARE\Classes\.uproject
-	FRegistryRootedKey *RootExtensionKey = new FRegistryRootedKey(HKEY_LOCAL_MACHINE, TEXT("SOFTWARE\\Classes\\.uproject"));
-	RootExtensionKey->Key = new FRegistryKey();
-	RootExtensionKey->Key->SetValue(TEXT(""), TEXT("Unreal.ProjectFile"));
-	RootedKeys.AddRawItem(RootExtensionKey);
-
-	// HKLM\SOFTWARE\Classes\Unreal.ProjectFile
-	FRegistryRootedKey *RootFileTypeKey = new FRegistryRootedKey(HKEY_LOCAL_MACHINE, TEXT("SOFTWARE\\Classes\\Unreal.ProjectFile"));
-	RootFileTypeKey->Key = new FRegistryKey();
-	RootFileTypeKey->Key->SetValue(TEXT(""), TEXT("Unreal Engine Project File"));
-	RootFileTypeKey->Key->FindOrAddKey(L"DefaultIcon")->SetValue(TEXT(""), QuotedExecutableFileName);
-	RootedKeys.AddRawItem(RootFileTypeKey);
-
-	// HKLM\SOFTWARE\Classes\Unreal.ProjectFile\shell
-	FRegistryKey *ShellKey = RootFileTypeKey->Key->FindOrAddKey(TEXT("shell"));
-
-	// HKLM\SOFTWARE\Classes\Unreal.ProjectFile\shell\open
-	FRegistryKey *ShellOpenKey = ShellKey->FindOrAddKey(TEXT("open"));
-	ShellOpenKey->SetValue(L"", L"Open");
-	ShellOpenKey->FindOrAddKey(L"command")->SetValue(L"", QuotedExecutableFileName + FString(TEXT(" /editor \"%1\"")));
-
-	// HKLM\SOFTWARE\Classes\Unreal.ProjectFile\shell\run
-	FRegistryKey *ShellRunKey = ShellKey->FindOrAddKey(TEXT("run"));
-	ShellRunKey->SetValue(TEXT(""), TEXT("Launch game"));
-	ShellRunKey->SetValue(TEXT("Icon"), QuotedExecutableFileName);
-	ShellRunKey->FindOrAddKey(L"command")->SetValue(TEXT(""), QuotedExecutableFileName + FString(TEXT(" /game \"%1\"")));
-
-	// HKLM\SOFTWARE\Classes\Unreal.ProjectFile\shell\rungenproj
-	FRegistryKey *ShellProjectKey = ShellKey->FindOrAddKey(TEXT("rungenproj"));
-	ShellProjectKey->SetValue(L"", L"Generate Visual Studio project files");
-	ShellProjectKey->SetValue(L"Icon", QuotedExecutableFileName);
-	ShellProjectKey->FindOrAddKey(L"command")->SetValue(TEXT(""), QuotedExecutableFileName + FString(TEXT(" /projectfiles \"%1\"")));
-
-	// HKLM\SOFTWARE\Classes\Unreal.ProjectFile\shell\switchversion
-	FRegistryKey *ShellVersionKey = ShellKey->FindOrAddKey(TEXT("switchversion"));
-	ShellVersionKey->SetValue(TEXT(""), TEXT("Switch engine version"));
-	ShellVersionKey->SetValue(TEXT("ExtendedSubCommandsKey"), TEXT("Unreal.ProjectFile\\switchversion"));
-	ShellVersionKey->SetValue(TEXT("Icon"), QuotedExecutableFileName);
-	ShellVersionKey->SetValue(TEXT("MUIVerb"), TEXT("Switch Unreal Engine version"));
-
-	// HKLM\SOFTWARE\Classes\Unreal.ProjectFile\switchversion
-	FRegistryKey *SwitchKey = RootFileTypeKey->Key->FindOrAddKey(TEXT("switchversion"));
-
-	// HKLM\SOFTWARE\Classes\Unreal.ProjectFile\switchversion\shell
-	FRegistryKey *SwitchShellKey = SwitchKey->FindOrAddKey(TEXT("shell"));
-
-	// HKLM\SOFTWARE\Classes\Unreal.ProjectFile\switchversion\shell\<Label>
+	FString Identifier;
+	TArray<FString> SortedIdentifiers;
 	TMap<FString, FString> Installations;
-	FDesktopPlatformModule::Get()->EnumerateEngineInstallations(Installations);
 
-	if (Installations.Num() > 0)
+	FSelectBuildDialog(const FString &InIdentifier)
 	{
-		// Add the default option
-		FRegistryKey *DefaultKey = SwitchShellKey->FindOrAddKey(L"000");
-		DefaultKey->SetValue(L"MUIVerb", L"Default");
-		DefaultKey->FindOrAddKey(L"command")->SetValue(L"", QuotedExecutableFileName + FString(TEXT(" /associate \"%1\" \"\"")));
-
-		// Build a list of installation names
-		TMap<FString, FString> InstallationNames;
-		for (TMap<FString, FString>::TConstIterator Iter(Installations); Iter; ++Iter)
-		{
-			const FString &Id = Iter->Key;
-			FString Description = GetInstallationDescription(Id, Iter.Value());
-			InstallationNames.Add(Description, Id);
-		}
-		InstallationNames.KeySort(FEngineLabelSortPredicate());
+		Identifier = InIdentifier;
+	}
+
+	bool DoModal(HWND hWndParent)
+	{
+		return DialogBoxParam(GetModuleHandle(NULL), MAKEINTRESOURCE(IDD_SELECTBUILD), hWndParent, &DialogProc, (LPARAM)this) != FALSE;
+	}
+
+private:
+	static INT_PTR CALLBACK DialogProc(HWND hWnd, UINT Msg, WPARAM wParam, LPARAM lParam)
+	{
+		FSelectBuildDialog *Dialog = (FSelectBuildDialog*)GetWindowLongPtr(hWnd, GWLP_USERDATA);
 
-		// Add the installation options
-		INT Idx = 1;
-		for (TMap<FString, FString>::TConstIterator Iter(InstallationNames); Iter; ++Iter)
+		switch (Msg)
 		{
-			wchar_t KeyName[100];
-			wsprintf(KeyName, L"%03d", Idx);
+		case WM_INITDIALOG:
+			Dialog = (FSelectBuildDialog*)lParam;
+			SetWindowLongPtr(hWnd, GWLP_USERDATA, (LONG_PTR)Dialog);
+			Dialog->UpdateInstallations(hWnd);
+			return false;
+		case WM_COMMAND:
+			switch (LOWORD(wParam))
+			{
+			case IDC_BROWSE:
+				if (Dialog->Browse(hWnd))
+				{
+					EndDialog(hWnd, 1);
+				}
+				return false;
+			case IDOK:
+				Dialog->StoreSelection(hWnd);
+				EndDialog(hWnd, 1);
+				return false;
+			case IDCANCEL:
+				EndDialog(hWnd, 0);
+				return false;
+			}
 
-			FRegistryKey *VersionKey = SwitchShellKey->FindOrAddKey(KeyName);
-			VersionKey->SetValue(TEXT("MUIVerb"), Iter.Key());
-			VersionKey->FindOrAddKey(TEXT("command"))->SetValue(TEXT(""), QuotedExecutableFileName + FString::Printf(TEXT(" /associate \"%%1\" \"%s\""), *Iter.Value()));
-			if(Idx++ == 1) VersionKey->SetValue(L"CommandFlags", 0x20);
 		}
+		return FALSE;
 	}
 
-	// HKLM\SOFTWARE\Classes\Unreal.ProjectFile\switchversion\shell\browse
-	FRegistryKey *BrowseShellKey = SwitchShellKey->FindOrAddKey(TEXT("browse"));
-	BrowseShellKey->SetValue(TEXT("MUIVerb"), TEXT("Browse..."));
-	BrowseShellKey->FindOrAddKey(TEXT("command"))->SetValue(L"", QuotedExecutableFileName + FString(TEXT(" /browse \"%1\"")));
-	if (Installations.Num() > 0)
+	void StoreSelection(HWND hWnd)
 	{
-		BrowseShellKey->SetValue(TEXT("CommandFlags"), 0x20);
+		int32 Idx = SendDlgItemMessage(hWnd, IDC_BUILDLIST, CB_GETCURSEL, 0, 0);
+		Identifier = SortedIdentifiers.IsValidIndex(Idx) ? SortedIdentifiers[Idx] : TEXT("");
 	}
-}
 
-bool FWindowsPlatformInstallation::GetLauncherVersionSelector(FString &OutFileName)
-{
-	// Get the launcher installed path
-	TCHAR InstallDir[MAX_PATH];
-	DWORD InstallDirLen = sizeof(InstallDir);
-	if (RegGetValue(HKEY_LOCAL_MACHINE, TEXT("SOFTWARE\\EpicGames\\Unreal Engine"), TEXT("INSTALLDIR"), RRF_RT_REG_SZ, NULL, InstallDir, &InstallDirLen) != ERROR_SUCCESS)
+	void UpdateInstallations(HWND hWnd)
 	{
-		return false;
-	}
+		FDesktopPlatformModule::Get()->EnumerateEngineInstallations(Installations);
+		Installations.GetKeys(SortedIdentifiers);
+		SortedIdentifiers.Sort<FEngineLabelSortPredicate>(FEngineLabelSortPredicate());
 
-	// Normalize the directory
-	FString NormalizedInstallDir = InstallDir;
-	FPaths::NormalizeDirectoryName(NormalizedInstallDir);
+		SendDlgItemMessage(hWnd, IDC_BUILDLIST, CB_RESETCONTENT, 0, 0);
 
-	// Check the version selector executable exists
-	FString LauncherFileName = NormalizedInstallDir / TEXT("Launcher/Engine/Binaries/Win64/UnrealVersionSelector.exe");
-	if (IFileManager::Get().FileSize(*LauncherFileName) < 0)
-	{
-		return false;
+		for(int32 Idx =  0; Idx < SortedIdentifiers.Num(); Idx++)
+		{
+			const FString &Identifier = SortedIdentifiers[Idx];
+			FString Description = GetInstallationDescription(Identifier, Installations[Identifier]);
+			SendDlgItemMessage(hWnd, IDC_BUILDLIST, CB_ADDSTRING, 0, (LPARAM)*Description);
+		}
+
+		int32 NewIdx = SortedIdentifiers.IndexOfByKey(Identifier);
+		SendDlgItemMessage(hWnd, IDC_BUILDLIST, CB_SETCURSEL, NewIdx, 0);
 	}
 
-	OutFileName = LauncherFileName;
-	return true;
-}
+	bool Browse(HWND hWnd)
+	{
+		// Get the currently bound engine directory for the project
+		const FString *RootDir = Installations.Find(Identifier);
+		FString EngineRootDir = (RootDir != NULL)? *RootDir : FString();
 
-bool FWindowsPlatformInstallation::RegisterEngineInstallation(const FString &Id, const FString &RootDirName)
-{
-	FString NormalizedRootDirName = RootDirName;
-	FPaths::NormalizeDirectoryName(NormalizedRootDirName);
+		// Browse for a new directory
+		FString NewEngineRootDir;
+		if (!FDesktopPlatformModule::Get()->OpenDirectoryDialog(hWnd, TEXT("Select the Unreal Engine installation to use for this project"), EngineRootDir, NewEngineRootDir))
+		{
+			return false;
+		}
 
-	FRegistryRootedKey RootKey(HKEY_LOCAL_MACHINE, FString::Printf(TEXT("SOFTWARE\\EpicGames\\Unreal Engine\\%s"), *Id));
-	RootKey.Key = new FRegistryKey();
-	RootKey.Key->SetValue(TEXT("InstalledDirectory"), NormalizedRootDirName);
-	return RootKey.Write(false);
-}
+		// Check it's a valid directory
+		if (!FPlatformInstallation::NormalizeEngineRootDir(NewEngineRootDir))
+		{
+			FPlatformMisc::MessageBoxExt(EAppMsgType::Ok, TEXT("The selected directory is not a valid engine installation."), TEXT("Error"));
+			return false;
+		}
 
-bool FWindowsPlatformInstallation::UnregisterEngineInstallation(const FString &Id)
-{
-	FRegistryRootedKey RootKey(HKEY_LOCAL_MACHINE, FString::Printf(TEXT("SOFTWARE\\EpicGames\\Unreal Engine\\%s"), *Id));
-	return RootKey.Write(true);
-}
+		// Check that it's a registered engine directory
+		FString NewIdentifier;
+		if (!FDesktopPlatformModule::Get()->GetEngineIdentifierFromRootDir(NewEngineRootDir, NewIdentifier))
+		{
+			if(!FPlatformInstallation::RegisterEngineInstallation(NewEngineRootDir, NewIdentifier))
+			{
+				FPlatformMisc::MessageBoxExt(EAppMsgType::Ok, TEXT("Couldn't register engine installation."), TEXT("Error"));
+				return false;
+			}
+		}
 
-bool FWindowsPlatformInstallation::IsSourceDistribution(const FString &EngineRootDir)
+		// Update the identifier and return
+		Identifier = NewIdentifier;
+		return true;
+	}
+};
+
+bool FWindowsPlatformInstallation::RegisterEngineInstallation(const FString &RootDirName, FString &OutIdentifier)
 {
-	// Check for the existence of a GenerateProjectFiles.bat file. This allows compatibility with the GitHub 4.0 release.
-	FString GenerateProjectFilesPath = EngineRootDir / TEXT("GenerateProjectFiles.bat");
-	if (IFileManager::Get().FileSize(*GenerateProjectFilesPath) >= 0)
+	OutIdentifier = FGuid::NewGuid().ToString(EGuidFormats::DigitsWithHyphensInBraces);
+
+	HKEY hRootKey;
+	if(RegCreateKeyEx(HKEY_CURRENT_USER, InstallationsSubKey, 0, NULL, REG_OPTION_NON_VOLATILE, KEY_ALL_ACCESS, NULL, &hRootKey, NULL) != ERROR_SUCCESS)
 	{
-		return true;
+		return false;
 	}
 
-	// Otherwise use the default test
-	return FGenericPlatformInstallation::IsSourceDistribution(EngineRootDir);
+	LRESULT SetResult = RegSetValueEx(hRootKey, *OutIdentifier, 0, REG_SZ, (const BYTE*)*RootDirName, (RootDirName.Len() + 1) * sizeof(TCHAR));
+	RegCloseKey(hRootKey);
+	return SetResult == ERROR_SUCCESS;
 }
 
 bool FWindowsPlatformInstallation::LaunchEditor(const FString &RootDirName, const FString &Arguments)
@@ -217,39 +177,43 @@ bool FWindowsPlatformInstallation::LaunchEditor(const FString &RootDirName, cons
 
 bool FWindowsPlatformInstallation::GenerateProjectFiles(const FString &RootDirName, const FString &Arguments)
 {
-	FString AllArguments = FString::Printf(TEXT("/c \"\"%s\" %s\""), *(RootDirName / TEXT("Engine/Build/BatchFiles/GenerateProjectFiles.bat")), *Arguments);
-	return FPlatformProcess::ExecProcess(TEXT("cmd.exe"), *AllArguments, NULL, NULL, NULL);
-}
-
-bool FWindowsPlatformInstallation::VerifyShellIntegration()
-{
-	TIndirectArray<FRegistryRootedKey> Keys;
-	GetRequiredRegistrySettings(Keys);
+	// Get the path to the batch file
+	FString BatchFileName;
+	if (FDesktopPlatformModule::Get()->IsSourceDistribution(RootDirName))
+	{
+		BatchFileName = RootDirName / TEXT("Engine/Build/BatchFiles/GenerateProjectFiles.bat");
+	}
+	else
+	{
+		BatchFileName = RootDirName / TEXT("Engine/Build/BatchFiles/RocketGenerateProjectFiles.bat");
+	}
 
-	for (int32 Idx = 0; Idx < Keys.Num(); Idx++)
+	// Check it exists
+	if (IFileManager::Get().FileSize(*BatchFileName) < 0)
 	{
-		if (!Keys[Idx].IsUpToDate(true))
-		{
-			return false;
-		}
+		return false;
 	}
 
+	// Build the full argument list
+	FString AllArguments = FString::Printf(TEXT("/c \"\"%s\" %s\""), *BatchFileName, *Arguments);
+
+	// Run the batch file through cmd.exe. 
+	FProcHandle Handle = FPlatformProcess::CreateProc(TEXT("cmd.exe"), *AllArguments, false, false, false, NULL, 0, NULL, NULL);
+	if (!Handle.IsValid()) return false;
+
+	FPlatformProcess::WaitForProc(Handle);
 	return true;
 }
 
-bool FWindowsPlatformInstallation::UpdateShellIntegration()
+bool FWindowsPlatformInstallation::SelectEngineInstallation(FString &Identifier)
 {
-	TIndirectArray<FRegistryRootedKey> Keys;
-	GetRequiredRegistrySettings(Keys);
-
-	for (int32 Idx = 0; Idx < Keys.Num(); Idx++)
+	FSelectBuildDialog Dialog(Identifier);
+	if(!Dialog.DoModal(NULL))
 	{
-		if (!Keys[Idx].Write(true))
-		{
-			return false;
-		}
+		return false;
 	}
 
+	Identifier = Dialog.Identifier;
 	return true;
 }
 
diff --git a/Engine/Source/Programs/UnrealVersionSelector/Private/Windows/WindowsPlatformInstallation.h b/Engine/Source/Programs/UnrealVersionSelector/Private/Windows/WindowsPlatformInstallation.h
--- a/Engine/Source/Programs/UnrealVersionSelector/Private/Windows/WindowsPlatformInstallation.h
+++ b/Engine/Source/Programs/UnrealVersionSelector/Private/Windows/WindowsPlatformInstallation.h
@@ -7,17 +7,8 @@
 
 struct FWindowsPlatformInstallation : FGenericPlatformInstallation
 {
-	/// Get the path to the installed build of the version selector
-	static bool GetLauncherVersionSelector(FString &OutFileName);
-
 	// Register an engine installation
-	static bool RegisterEngineInstallation(const FString &Id, const FString &RootDirName);
-
-	// Unregister an engine installation
-	static bool UnregisterEngineInstallation(const FString &Id);
-
-	// Test whether a given engine root directory is a source build
-	static bool IsSourceDistribution(const FString &EngineRootDir);
+	static bool RegisterEngineInstallation(const FString &RootDirName, FString &OutIdentifier);
 
 	// Launches the editor application
 	static bool LaunchEditor(const FString &RootDirName, const FString &Arguments);
@@ -25,10 +16,7 @@ struct FWindowsPlatformInstallation : FGenericPlatformInstallation
 	// Generate project files by running UBT
 	static bool GenerateProjectFiles(const FString &RootDirName, const FString &Arguments);
 
-	// Check whether the current settings are up to date
-	static bool VerifyShellIntegration();
-
-	// Update the installation settings
-	static bool UpdateShellIntegration();
+	// Select an engine installation
+	static bool SelectEngineInstallation(FString &Identifier);
 };
 
diff --git a/Engine/Source/Programs/UnrealVersionSelector/Private/Windows/WindowsRegistry.cpp b/Engine/Source/Programs/UnrealVersionSelector/Private/Windows/WindowsRegistry.cpp
deleted file mode 100644
diff --git a/Engine/Source/Programs/UnrealVersionSelector/Private/Windows/WindowsRegistry.h b/Engine/Source/Programs/UnrealVersionSelector/Private/Windows/WindowsRegistry.h
deleted file mode 100644
