Update launcher to use new file association code. Include UnrealVersionSelector in the launcher build process.

diff --git a/Engine/Source/Programs/UnrealVersionSelector/UnrealVersionSelector.Target.cs b/Engine/Source/Programs/UnrealVersionSelector/UnrealVersionSelector.Target.cs
--- a/Engine/Source/Programs/UnrealVersionSelector/UnrealVersionSelector.Target.cs
+++ b/Engine/Source/Programs/UnrealVersionSelector/UnrealVersionSelector.Target.cs
@@ -55,8 +55,12 @@ public class UnrealVersionSelectorTarget : TargetRules
 	}
 	public override bool GUBP_AlwaysBuildWithTools(UnrealTargetPlatform InHostPlatform, out bool bInternalToolOnly, out bool SeparateNode)
 	{
-		bInternalToolOnly = false;
-		SeparateNode = false;
+		bInternalToolOnly = true;
+		SeparateNode = true;
 		return true;
 	}
+	public override List<UnrealTargetConfiguration> GUBP_ToolConfigs(UnrealTargetPlatform InHostPlatform)
+	{
+		return new List<UnrealTargetConfiguration> { UnrealTargetConfiguration.Shipping };
+	}
 }
