Fix command line for launching games using UnrealVersionSelector. Engine arguments use hyphens, not slashes.

diff --git a/Engine/Source/Programs/UnrealVersionSelector/Private/UnrealVersionSelector.cpp b/Engine/Source/Programs/UnrealVersionSelector/Private/UnrealVersionSelector.cpp
--- a/Engine/Source/Programs/UnrealVersionSelector/Private/UnrealVersionSelector.cpp
+++ b/Engine/Source/Programs/UnrealVersionSelector/Private/UnrealVersionSelector.cpp
@@ -178,7 +178,7 @@ int Main(const TArray<FString> &Arguments)
 	else if (Arguments.Num() == 2 && Arguments[0] == TEXT("/game"))
 	{
 		// Play a game using the editor executable
-		bRes = LaunchEditor(Arguments[1], L"/game");
+		bRes = LaunchEditor(Arguments[1], L"-game");
 	}
 	else if (Arguments.Num() == 2 && Arguments[0] == TEXT("/projectfiles"))
 	{
