Fix UnrealVersionSelector attempting to generate project files when switching version for a content-only project.

diff --git a/Engine/Source/Programs/UnrealVersionSelector/Private/UnrealVersionSelector.cpp b/Engine/Source/Programs/UnrealVersionSelector/Private/UnrealVersionSelector.cpp
--- a/Engine/Source/Programs/UnrealVersionSelector/Private/UnrealVersionSelector.cpp
+++ b/Engine/Source/Programs/UnrealVersionSelector/Private/UnrealVersionSelector.cpp
@@ -75,7 +75,14 @@ bool SwitchVersion(const FString& ProjectFileName)
 		return false;
 	}
 
-	// Clean all the build products
+	// If it's a content-only project, we're done
+	FProjectStatus ProjectStatus;
+	if(IProjectManager::Get().QueryStatusForProject(ProjectFileName, ProjectStatus) && !ProjectStatus.bCodeBasedProject)
+	{
+		return true;
+	}
+
+	// Otherwise clean all the build products
 	FFeedbackContext *Warn = FDesktopPlatformModule::Get()->GetNativeFeedbackContext();
 	if (!FDesktopPlatformModule::Get()->CleanGameProject(FPaths::GetPath(ProjectFileName), Warn))
 	{
