Prompt to browse for an engine installation if there isn't one installed.

diff --git a/Engine/Source/Programs/UnrealVersionSelector/Private/Windows/WindowsPlatformInstallation.cpp b/Engine/Source/Programs/UnrealVersionSelector/Private/Windows/WindowsPlatformInstallation.cpp
--- a/Engine/Source/Programs/UnrealVersionSelector/Private/Windows/WindowsPlatformInstallation.cpp
+++ b/Engine/Source/Programs/UnrealVersionSelector/Private/Windows/WindowsPlatformInstallation.cpp
@@ -43,20 +43,33 @@ FString GetInstallationDescription(const FString &Id, const FString &RootDir)
 struct FSelectBuildDialog
 {
 	FString Identifier;
-	TArray<FString> SortedIdentifiers;
-	TMap<FString, FString> Installations;
 
 	FSelectBuildDialog(const FString &InIdentifier)
 	{
 		Identifier = InIdentifier;
+		FDesktopPlatformModule::Get()->EnumerateEngineInstallations(Installations);
+		Installations.GetKeys(SortedIdentifiers);
+		SortedIdentifiers.Sort<FEngineLabelSortPredicate>(FEngineLabelSortPredicate());
 	}
 
 	bool DoModal(HWND hWndParent)
 	{
-		return DialogBoxParam(GetModuleHandle(NULL), MAKEINTRESOURCE(IDD_SELECTBUILD), hWndParent, &DialogProc, (LPARAM)this) != FALSE;
+		// If there's more than one already, select them from a list
+		if(Installations.Num() > 0)
+		{
+			return DialogBoxParam(GetModuleHandle(NULL), MAKEINTRESOURCE(IDD_SELECTBUILD), hWndParent, &DialogProc, (LPARAM)this) != FALSE;
+		}
+		else if(FPlatformMisc::MessageBoxExt(EAppMsgType::YesNo, TEXT("No Unreal Engine installations found. Would you like to locate one manually?"), TEXT("Installation Not Found")) == EAppReturnType::Yes && Browse(hWndParent))
+		{
+			return true;
+		}
+		return false;
 	}
 
 private:
+	TArray<FString> SortedIdentifiers;
+	TMap<FString, FString> Installations;
+
 	static INT_PTR CALLBACK DialogProc(HWND hWnd, UINT Msg, WPARAM wParam, LPARAM lParam)
 	{
 		FSelectBuildDialog *Dialog = (FSelectBuildDialog*)GetWindowLongPtr(hWnd, GWLP_USERDATA);
@@ -98,10 +111,6 @@ private:
 
 	void UpdateInstallations(HWND hWnd)
 	{
-		FDesktopPlatformModule::Get()->EnumerateEngineInstallations(Installations);
-		Installations.GetKeys(SortedIdentifiers);
-		SortedIdentifiers.Sort<FEngineLabelSortPredicate>(FEngineLabelSortPredicate());
-
 		SendDlgItemMessage(hWnd, IDC_BUILDLIST, CB_RESETCONTENT, 0, 0);
 
 		for(int32 Idx =  0; Idx < SortedIdentifiers.Num(); Idx++)
