Add a simple error message when trying to generate code projects for content-only projects.

diff --git a/Engine/Source/Programs/UnrealVersionSelector/Private/UnrealVersionSelector.cpp b/Engine/Source/Programs/UnrealVersionSelector/Private/UnrealVersionSelector.cpp
--- a/Engine/Source/Programs/UnrealVersionSelector/Private/UnrealVersionSelector.cpp
+++ b/Engine/Source/Programs/UnrealVersionSelector/Private/UnrealVersionSelector.cpp
@@ -140,6 +140,14 @@ bool GenerateProjectFiles(const FString& ProjectFileName)
 {
 	IDesktopPlatform* DesktopPlatform = FDesktopPlatformModule::Get();
 
+	// Check it's a code project
+	FString SourceDir = FPaths::GetPath(ProjectFileName) / TEXT("Source");
+	if(!IPlatformFile::GetPlatformPhysical().DirectoryExists(*SourceDir))
+	{
+		FPlatformMisc::MessageBoxExt(EAppMsgType::Ok, TEXT("This project does not have any source code. You need to add C++ source files to the project from the Editor before you can generate project files."), TEXT("Error"));
+		return false;
+	}
+
 	// Get the engine root directory
 	FString RootDir;
 	if (!GetValidatedEngineRootDir(ProjectFileName, RootDir))
