Add an option to UnrealVersionSelector which will register the engine directory without throwing up a popup message.

diff --git a/Engine/Source/Programs/UnrealVersionSelector/Private/UnrealVersionSelector.cpp b/Engine/Source/Programs/UnrealVersionSelector/Private/UnrealVersionSelector.cpp
--- a/Engine/Source/Programs/UnrealVersionSelector/Private/UnrealVersionSelector.cpp
+++ b/Engine/Source/Programs/UnrealVersionSelector/Private/UnrealVersionSelector.cpp
@@ -7,14 +7,8 @@ IMPLEMENT_APPLICATION(UnrealVersionSelector, "UnrealVersionSelector")
 
 bool GenerateProjectFiles(const FString& ProjectFileName);
 
-bool RegisterCurrentEngineDirectory()
+bool RegisterCurrentEngineDirectory(bool bPromptForFileAssociations)
 {
-	// Prompt for registering this directory
-	if(FPlatformMisc::MessageBoxExt(EAppMsgType::YesNo, TEXT("Register this directory as an Unreal Engine installation?"), TEXT("Question")) != EAppReturnType::Yes)
-	{
-		return false;
-	}
-
 	// Get the current engine directory.
 	FString EngineRootDir = FPlatformProcess::BaseDir();
 	if(!FPlatformInstallation::NormalizeEngineRootDir(EngineRootDir))
@@ -32,18 +26,39 @@ bool RegisterCurrentEngineDirectory()
 	}
 
 	// If the launcher isn't installed, set up the file associations
-	if(!FDesktopPlatformModule::Get()->VerifyFileAssociations())
+	if(!FDesktopPlatformModule::Get()->VerifyFileAssociations() || true)
 	{
-		// Relaunch as administrator
-		FString ExecutableFileName = FString(FPlatformProcess::BaseDir()) / FString(FPlatformProcess::ExecutableName(false));
-
-		int32 ExitCode;
-		if (!FPlatformProcess::ExecElevatedProcess(*ExecutableFileName, TEXT("/fileassociations"), &ExitCode) || ExitCode != 0)
+		// Prompt for whether to update the file associations
+		if(!bPromptForFileAssociations || FPlatformMisc::MessageBoxExt(EAppMsgType::YesNo, TEXT("Register Unreal Engine file types?"), TEXT("File Types")) == EAppReturnType::Yes)
 		{
-			return false;
+			// Relaunch as administrator
+			FString ExecutableFileName = FString(FPlatformProcess::BaseDir()) / FString(FPlatformProcess::ExecutableName(false));
+
+			int32 ExitCode;
+			if (!FPlatformProcess::ExecElevatedProcess(*ExecutableFileName, TEXT("/fileassociations"), &ExitCode) || ExitCode != 0)
+			{
+				return false;
+			}
 		}
 	}
 
+	return true;
+}
+
+bool RegisterCurrentEngineDirectoryWithPrompt()
+{
+	// Ask whether the user wants to register the directory
+	if(FPlatformMisc::MessageBoxExt(EAppMsgType::YesNo, TEXT("Register this directory as an Unreal Engine installation?"), TEXT("Question")) != EAppReturnType::Yes)
+	{
+		return false;
+	}
+
+	// Register the engine directory. We've already prompted for registering the directory, so 
+	if(!RegisterCurrentEngineDirectory(false))
+	{
+		return false;
+	}
+
 	// Notify the user that everything is awesome.
 	FPlatformMisc::MessageBoxExt(EAppMsgType::Ok, TEXT("Registration successful."), TEXT("Success"));
 	return true;
@@ -188,7 +203,12 @@ int Main(const TArray<FString>& Arguments)
 	if (Arguments.Num() == 0)
 	{
 		// Add the current directory to the list of installations
-		bRes = RegisterCurrentEngineDirectory();
+		bRes = RegisterCurrentEngineDirectoryWithPrompt();
+	}
+	else if (Arguments.Num() == 1 && Arguments[0] == TEXT("/register"))
+	{
+		// Add the current directory to the list of installations
+		bRes = RegisterCurrentEngineDirectory(true);
 	}
 	else if (Arguments.Num() == 1 && Arguments[0] == TEXT("/fileassociations"))
 	{
