Do not build UnrealVersionSelector on Mac

diff --git a/Engine/Source/Programs/UnrealVersionSelector/UnrealVersionSelector.Target.cs b/Engine/Source/Programs/UnrealVersionSelector/UnrealVersionSelector.Target.cs
--- a/Engine/Source/Programs/UnrealVersionSelector/UnrealVersionSelector.Target.cs
+++ b/Engine/Source/Programs/UnrealVersionSelector/UnrealVersionSelector.Target.cs
@@ -54,9 +54,16 @@ public class UnrealVersionSelectorTarget : TargetRules
 	}
     public override bool GUBP_AlwaysBuildWithTools(UnrealTargetPlatform InHostPlatform, bool bBuildingRocket, out bool bInternalToolOnly, out bool SeparateNode)
 	{
-		bInternalToolOnly = true;
-		SeparateNode = true;
-		return true;
+		if (InHostPlatform == UnrealTargetPlatform.Win32 || InHostPlatform == UnrealTargetPlatform.Win64)
+		{
+			bInternalToolOnly = true;
+			SeparateNode = false;
+			return true;
+		}
+
+		bInternalToolOnly = false;
+		SeparateNode = false;
+		return false;
 	}
 	public override List<UnrealTargetConfiguration> GUBP_ToolConfigs(UnrealTargetPlatform InHostPlatform)
 	{
