build script, fix dependencies for not building unrealversionselector on mac

diff --git a/Engine/Source/Programs/UnrealVersionSelector/UnrealVersionSelector.Target.cs b/Engine/Source/Programs/UnrealVersionSelector/UnrealVersionSelector.Target.cs
--- a/Engine/Source/Programs/UnrealVersionSelector/UnrealVersionSelector.Target.cs
+++ b/Engine/Source/Programs/UnrealVersionSelector/UnrealVersionSelector.Target.cs
@@ -57,7 +57,7 @@ public class UnrealVersionSelectorTarget : TargetRules
 		if (InHostPlatform == UnrealTargetPlatform.Win32 || InHostPlatform == UnrealTargetPlatform.Win64)
 		{
 			bInternalToolOnly = true;
-			SeparateNode = false;
+			SeparateNode = true;
 			return true;
 		}
 
