Fix UnrealVersionSelector spawning the editor behind other processes. FWindowsPlatformProcess::ExecProcess() passes SW_SHOWMINNOACTIVE in STARTUPINFO, so just call CreateProcess() directly instead. (Merged back from 4.6 branch).

diff --git a/Engine/Source/Programs/UnrealVersionSelector/Private/UnrealVersionSelector.cpp b/Engine/Source/Programs/UnrealVersionSelector/Private/UnrealVersionSelector.cpp
--- a/Engine/Source/Programs/UnrealVersionSelector/Private/UnrealVersionSelector.cpp
+++ b/Engine/Source/Programs/UnrealVersionSelector/Private/UnrealVersionSelector.cpp
@@ -207,32 +207,32 @@ int Main(const TArray<FString>& Arguments)
 		// Add the current directory to the list of installations
 		bRes = RegisterCurrentEngineDirectoryWithPrompt();
 	}
-	else if (Arguments.Num() == 1 && Arguments[0] == TEXT("/register"))
+	else if (Arguments.Num() == 1 && Arguments[0] == TEXT("-register"))
 	{
 		// Add the current directory to the list of installations
 		bRes = RegisterCurrentEngineDirectory(true);
 	}
-	else if (Arguments.Num() == 1 && Arguments[0] == TEXT("/fileassociations"))
+	else if (Arguments.Num() == 1 && Arguments[0] == TEXT("-fileassociations"))
 	{
 		// Update all the settings.
 		bRes = UpdateFileAssociations();
 	}
-	else if (Arguments.Num() == 2 && Arguments[0] == TEXT("/switchversion"))
+	else if (Arguments.Num() == 2 && Arguments[0] == TEXT("-switchversion"))
 	{
 		// Associate with an engine label
 		bRes = SwitchVersion(Arguments[1]);
 	}
-	else if (Arguments.Num() == 2 && Arguments[0] == TEXT("/editor"))
+	else if (Arguments.Num() == 2 && Arguments[0] == TEXT("-editor"))
 	{
 		// Open a project with the editor
 		bRes = LaunchEditor(Arguments[1], L"");
 	}
-	else if (Arguments.Num() == 2 && Arguments[0] == TEXT("/game"))
+	else if (Arguments.Num() == 2 && Arguments[0] == TEXT("-game"))
 	{
 		// Play a game using the editor executable
 		bRes = LaunchEditor(Arguments[1], L"-game");
 	}
-	else if (Arguments.Num() == 2 && Arguments[0] == TEXT("/projectfiles"))
+	else if (Arguments.Num() == 2 && Arguments[0] == TEXT("-projectfiles"))
 	{
 		// Generate Visual Studio project files
 		bRes = GenerateProjectFiles(Arguments[1]);
@@ -260,7 +260,12 @@ int Main(const TArray<FString>& Arguments)
 		TArray<FString> Arguments;
 		for (int Idx = 1; Idx < ArgC; Idx++)
 		{
-			Arguments.Add(ArgV[Idx]);
+			FString Argument = ArgV[Idx];
+			if(Argument.Len() > 0 && Argument[0] == '/')
+			{
+				Argument[0] = '-';
+			}
+			Arguments.Add(Argument);
 		}
 
 		return Main(Arguments);
diff --git a/Engine/Source/Programs/UnrealVersionSelector/Private/Windows/WindowsPlatformInstallation.cpp b/Engine/Source/Programs/UnrealVersionSelector/Private/Windows/WindowsPlatformInstallation.cpp
--- a/Engine/Source/Programs/UnrealVersionSelector/Private/Windows/WindowsPlatformInstallation.cpp
+++ b/Engine/Source/Programs/UnrealVersionSelector/Private/Windows/WindowsPlatformInstallation.cpp
@@ -214,8 +214,21 @@ private:
 
 bool FWindowsPlatformInstallation::LaunchEditor(const FString &RootDirName, const FString &Arguments)
 {
-	FString EditorFileName = RootDirName / TEXT("Engine/Binaries/Win64/UE4Editor.exe");
-	return FPlatformProcess::ExecProcess(*EditorFileName, *Arguments, NULL, NULL, NULL);
+	FString CommandLine = FString::Printf(TEXT("\"%s\" %s"), *(RootDirName / TEXT("Engine/Binaries/Win64/UE4Editor.exe")), *Arguments);
+
+	STARTUPINFO StartupInfo;
+	ZeroMemory(&StartupInfo, sizeof(StartupInfo));
+	StartupInfo.cb = sizeof(StartupInfo);
+
+	PROCESS_INFORMATION ProcessInfo;
+	if(!CreateProcess(NULL, CommandLine.GetCharArray().GetData(), NULL, NULL, FALSE, 0, NULL, NULL, &StartupInfo, &ProcessInfo))
+	{
+		return false;
+	}
+
+	CloseHandle(ProcessInfo.hThread);
+	CloseHandle(ProcessInfo.hProcess);
+	return true;
 }
 
 bool FWindowsPlatformInstallation::SelectEngineInstallation(FString &Identifier)
