Remove debug code that was accidentally submitted. UnrealVersionSelector shouldn't update file associations if they're not out of date.

diff --git a/Engine/Source/Programs/UnrealVersionSelector/Private/UnrealVersionSelector.cpp b/Engine/Source/Programs/UnrealVersionSelector/Private/UnrealVersionSelector.cpp
--- a/Engine/Source/Programs/UnrealVersionSelector/Private/UnrealVersionSelector.cpp
+++ b/Engine/Source/Programs/UnrealVersionSelector/Private/UnrealVersionSelector.cpp
@@ -28,7 +28,7 @@ bool RegisterCurrentEngineDirectory(bool bPromptForFileAssociations)
 	}
 
 	// If the launcher isn't installed, set up the file associations
-	if(!FDesktopPlatformModule::Get()->VerifyFileAssociations() || true)
+	if(!FDesktopPlatformModule::Get()->VerifyFileAssociations())
 	{
 		// Prompt for whether to update the file associations
 		if(!bPromptForFileAssociations || FPlatformMisc::MessageBoxExt(EAppMsgType::YesNo, TEXT("Register Unreal Engine file types?"), TEXT("File Types")) == EAppReturnType::Yes)
