Remove the separate node for UnrealVersionSelector - just build it with tools, and only on Win64.

diff --git a/Engine/Source/Programs/AutomationTool/Scripts/GUBP.Automation.cs b/Engine/Source/Programs/AutomationTool/Scripts/GUBP.Automation.cs
--- a/Engine/Source/Programs/AutomationTool/Scripts/GUBP.Automation.cs
+++ b/Engine/Source/Programs/AutomationTool/Scripts/GUBP.Automation.cs
@@ -1901,15 +1901,6 @@ public class GUBP : BuildCommand
             AddDependency(RootEditorNode.StaticGetFullName(HostPlatform));
             AddDependency(ToolsNode.StaticGetFullName(HostPlatform));
             AddDependency(InternalToolsNode.StaticGetFullName(HostPlatform));
-
-			if(InHostPlatform == UnrealTargetPlatform.Win64)
-			{
-				var VersionSelector = bp.Branch.FindProgram("UnrealVersionSelector"); 
-				if (VersionSelector.Rules != null) 
-				{ 
-					AddDependency(SingleInternalToolsNode.StaticGetFullName(HostPlatform, VersionSelector));
-				}
-			}
         }
         public static string StaticGetFullName(UnrealTargetPlatform InHostPlatform)
         {
diff --git a/Engine/Source/Programs/AutomationTool/Scripts/RocketBuild.Automation.cs b/Engine/Source/Programs/AutomationTool/Scripts/RocketBuild.Automation.cs
--- a/Engine/Source/Programs/AutomationTool/Scripts/RocketBuild.Automation.cs
+++ b/Engine/Source/Programs/AutomationTool/Scripts/RocketBuild.Automation.cs
@@ -255,12 +255,6 @@ namespace Rocket
 				AddDependency(GUBP.ToolsNode.StaticGetFullName(OtherHostPlatform));
 				AddDependency(GUBP.InternalToolsNode.StaticGetFullName(OtherHostPlatform));
 			}
-
-			SingleTargetProperties VersionSelector = bp.Branch.FindProgram("UnrealVersionSelector"); 
-			if (VersionSelector.Rules != null) 
-			{ 
-				AddDependency(GUBP.SingleInternalToolsNode.StaticGetFullName(UnrealTargetPlatform.Win64, VersionSelector));
-			}
 		}
 
 		public static string StaticGetFullName(UnrealTargetPlatform HostPlatform)
diff --git a/Engine/Source/Programs/UnrealVersionSelector/UnrealVersionSelector.Target.cs b/Engine/Source/Programs/UnrealVersionSelector/UnrealVersionSelector.Target.cs
--- a/Engine/Source/Programs/UnrealVersionSelector/UnrealVersionSelector.Target.cs
+++ b/Engine/Source/Programs/UnrealVersionSelector/UnrealVersionSelector.Target.cs
@@ -55,16 +55,9 @@ public class UnrealVersionSelectorTarget : TargetRules
     public override bool GUBP_AlwaysBuildWithTools(UnrealTargetPlatform InHostPlatform, out bool bInternalToolOnly, out bool SeparateNode, out bool CrossCompile)
 	{
 		CrossCompile = false;
-		if (InHostPlatform == UnrealTargetPlatform.Win32 || InHostPlatform == UnrealTargetPlatform.Win64)
-		{
-			bInternalToolOnly = true;
-			SeparateNode = true;
-			return true;
-		}
-
 		bInternalToolOnly = false;
 		SeparateNode = false;
-		return false;
+		return InHostPlatform == UnrealTargetPlatform.Win64;
 	}
 	public override List<UnrealTargetConfiguration> GUBP_ToolConfigs(UnrealTargetPlatform InHostPlatform)
 	{
