Fix line endings in UnrealVersionSelector resource file. Causing compile errors to be ignored, since all lines are treated as part of the comment  at the start of the file.

diff --git a/Engine/Source/Programs/UnrealVersionSelector/Private/Windows/Resources/Resource.rc b/Engine/Source/Programs/UnrealVersionSelector/Private/Windows/Resources/Resource.rc
new file mode 100644
--- /dev/null
+++ b/Engine/Source/Programs/UnrealVersionSelector/Private/Windows/Resources/Resource.rc
@@ -0,0 +1,116 @@
+// Microsoft Visual C++ generated resource script.
+//
+#include "resource.h"
+
+#define APSTUDIO_READONLY_SYMBOLS
+/////////////////////////////////////////////////////////////////////////////
+//
+// Generated from the TEXTINCLUDE 2 resource.
+//
+#include "winres.h"
+#include "VersionResource.inl"
+
+/////////////////////////////////////////////////////////////////////////////
+#undef APSTUDIO_READONLY_SYMBOLS
+
+/////////////////////////////////////////////////////////////////////////////
+// English (United States) resources
+
+#if !defined(AFX_RESOURCE_DLL) || defined(AFX_TARG_ENU)
+LANGUAGE LANG_ENGLISH, SUBLANG_ENGLISH_US
+
+#ifdef APSTUDIO_INVOKED
+/////////////////////////////////////////////////////////////////////////////
+//
+// TEXTINCLUDE
+//
+
+1 TEXTINCLUDE 
+BEGIN
+    "resource.h\0"
+END
+
+2 TEXTINCLUDE 
+BEGIN
+    "#include ""winres.h""\r\n"
+    "#include ""VersionResource.inl""\r\n"
+    "\0"
+END
+
+3 TEXTINCLUDE 
+BEGIN
+    "\r\n"
+    "\0"
+END
+
+#endif    // APSTUDIO_INVOKED
+
+
+/////////////////////////////////////////////////////////////////////////////
+//
+// Dialog
+//
+
+IDD_SELECTBUILD DIALOGEX 0, 0, 301, 47
+STYLE DS_SYSMODAL | DS_SETFONT | DS_MODALFRAME | DS_FIXEDSYS | DS_CENTER | WS_POPUP | WS_VISIBLE | WS_CAPTION | WS_SYSMENU
+CAPTION "Select Unreal Engine Version"
+FONT 8, "MS Shell Dlg", 400, 0, 0x1
+BEGIN
+    DEFPUSHBUTTON   "OK",IDOK,191,26,50,14
+    PUSHBUTTON      "Cancel",IDCANCEL,244,26,50,14
+    PUSHBUTTON      "...",IDC_BROWSE,270,6,24,14
+    COMBOBOX        IDC_BUILDLIST,7,7,259,135,CBS_DROPDOWNLIST | WS_VSCROLL | WS_TABSTOP
+END
+
+IDD_ERRORDIALOG DIALOGEX 0, 0, 541, 166
+STYLE DS_SETFONT | DS_MODALFRAME | DS_SETFOREGROUND | DS_FIXEDSYS | DS_CENTER | WS_POPUP | WS_CAPTION | WS_SYSMENU
+CAPTION "Error"
+FONT 8, "MS Shell Dlg", 400, 0, 0x1
+BEGIN
+    DEFPUSHBUTTON   "OK",IDOK,484,145,50,14
+    LTEXT           "Static",IDC_ERRORMESSAGE,7,7,527,8
+    EDITTEXT        IDC_ERRORLOGTEXT,7,19,527,124,ES_MULTILINE | ES_AUTOVSCROLL | ES_AUTOHSCROLL | ES_READONLY | WS_VSCROLL | WS_HSCROLL
+END
+
+
+/////////////////////////////////////////////////////////////////////////////
+//
+// DESIGNINFO
+//
+
+#ifdef APSTUDIO_INVOKED
+GUIDELINES DESIGNINFO
+BEGIN
+    IDD_SELECTBUILD, DIALOG
+    BEGIN
+        LEFTMARGIN, 7
+        RIGHTMARGIN, 294
+        TOPMARGIN, 7
+        BOTTOMMARGIN, 40
+    END
+
+    IDD_ERRORDIALOG, DIALOG
+    BEGIN
+        LEFTMARGIN, 7
+        RIGHTMARGIN, 534
+        TOPMARGIN, 7
+        BOTTOMMARGIN, 159
+    END
+END
+#endif    // APSTUDIO_INVOKED
+
+#endif    // English (United States) resources
+/////////////////////////////////////////////////////////////////////////////
+
+
+
+#ifndef APSTUDIO_INVOKED
+/////////////////////////////////////////////////////////////////////////////
+//
+// Generated from the TEXTINCLUDE 3 resource.
+//
+
+
+/////////////////////////////////////////////////////////////////////////////
+#endif    // not APSTUDIO_INVOKED
+
