Fix path to resources for UnrealVersionSelector.

diff --git a/Engine/Source/Programs/UnrealVersionSelector/Private/Windows/Resources/VersionResource.inl b/Engine/Source/Programs/UnrealVersionSelector/Private/Windows/Resources/VersionResource.inl
--- a/Engine/Source/Programs/UnrealVersionSelector/Private/Windows/Resources/VersionResource.inl
+++ b/Engine/Source/Programs/UnrealVersionSelector/Private/Windows/Resources/VersionResource.inl
@@ -56,18 +56,18 @@ END
 
 // Debug manifest
 #if _WIN64
-1						RT_MANIFEST				"../../../../../../Build/Windows/Default-Win64-Debug.manifest"
+1						RT_MANIFEST				"../../../../../../Build/Windows/Resources/Default-Win64-Debug.manifest"
 #else
-1						RT_MANIFEST				"../../../../../../Build/Windows/Default-Win32-Debug.manifest"
+1						RT_MANIFEST				"../../../../../../Build/Windows/Resources/Default-Win32-Debug.manifest"
 #endif
 
 #else
 
 // Release manifest
 #if _WIN64
-1						RT_MANIFEST				"../../../../../../Build/Windows/Default-Win64.manifest"
+1						RT_MANIFEST				"../../../../../../Build/Windows/Resources/Default-Win64.manifest"
 #else
-1						RT_MANIFEST				"../../../../../../Build/Windows/Default-Win32.manifest"
+1						RT_MANIFEST				"../../../../../../Build/Windows/Resources/Default-Win32.manifest"
 #endif
 
 #endif
