Allow UnrealVersionSelector to be called with -register -unattended to avoid the popup but still register the events

diff --git a/Engine/Source/Programs/UnrealVersionSelector/Private/UnrealVersionSelector.cpp b/Engine/Source/Programs/UnrealVersionSelector/Private/UnrealVersionSelector.cpp
--- a/Engine/Source/Programs/UnrealVersionSelector/Private/UnrealVersionSelector.cpp
+++ b/Engine/Source/Programs/UnrealVersionSelector/Private/UnrealVersionSelector.cpp
@@ -263,6 +263,11 @@ int Main(const TArray<FString>& Arguments)
 		// Add the current directory to the list of installations
 		bRes = RegisterCurrentEngineDirectory(true);
 	}
+	else if (Arguments.Num() == 2 && Arguments[0] == TEXT("-register") && Arguments[1] == TEXT("-unattended"))
+	{
+		// Add the current directory to the list of installations
+		bRes = RegisterCurrentEngineDirectory(false);
+	}
 	else if (Arguments.Num() == 1 && Arguments[0] == TEXT("-fileassociations"))
 	{
 		// Update all the settings.
