UnrealVersionSelector: Add support for projects built with a custom build environment. Enumerate all the .target files under a project's binaries directory, and take the launch path from the most recent one which is Configuration = "Development" and Type = "Editor".

diff --git a/Engine/Source/Programs/UnrealVersionSelector/Private/GenericPlatform/GenericPlatformInstallation.cpp b/Engine/Source/Programs/UnrealVersionSelector/Private/GenericPlatform/GenericPlatformInstallation.cpp
--- a/Engine/Source/Programs/UnrealVersionSelector/Private/GenericPlatform/GenericPlatformInstallation.cpp
+++ b/Engine/Source/Programs/UnrealVersionSelector/Private/GenericPlatform/GenericPlatformInstallation.cpp
@@ -37,7 +37,7 @@ bool FGenericPlatformInstallation::NormalizeEngineRootDir(FString &RootDir)
 	return false;
 }
 
-bool FGenericPlatformInstallation::LaunchEditor(const FString &RootDirName, const FString &Arguments)
+bool FGenericPlatformInstallation::LaunchEditor(const FString& RootDirName, const FString &FileName, const FString &Arguments)
 {
 	// No default implementation
 	return false;
diff --git a/Engine/Source/Programs/UnrealVersionSelector/Private/GenericPlatform/GenericPlatformInstallation.h b/Engine/Source/Programs/UnrealVersionSelector/Private/GenericPlatform/GenericPlatformInstallation.h
--- a/Engine/Source/Programs/UnrealVersionSelector/Private/GenericPlatform/GenericPlatformInstallation.h
+++ b/Engine/Source/Programs/UnrealVersionSelector/Private/GenericPlatform/GenericPlatformInstallation.h
@@ -10,7 +10,7 @@ struct FGenericPlatformInstallation
 	static bool NormalizeEngineRootDir(FString &RootDir);
 
 	// Launches the editor application
-	static bool LaunchEditor(const FString &RootDirName, const FString &Arguments);
+	static bool LaunchEditor(const FString& RootDirName, const FString &ExplicitFileName, const FString &Arguments);
 
 	// Select an engine installation
 	static bool SelectEngineInstallation(FString &Identifier);
diff --git a/Engine/Source/Programs/UnrealVersionSelector/Private/Linux/LinuxPlatformInstallation.cpp b/Engine/Source/Programs/UnrealVersionSelector/Private/Linux/LinuxPlatformInstallation.cpp
--- a/Engine/Source/Programs/UnrealVersionSelector/Private/Linux/LinuxPlatformInstallation.cpp
+++ b/Engine/Source/Programs/UnrealVersionSelector/Private/Linux/LinuxPlatformInstallation.cpp
@@ -504,11 +504,19 @@ private:
 	}
 };
 
-bool FLinuxPlatformInstallation::LaunchEditor(const FString &RootDirName, const FString &Arguments)
+bool FLinuxPlatformInstallation::LaunchEditor(const FString& RootDirName, const FString &ExplicitFileName, const FString &Arguments)
 {
-	FString CommandLine = RootDirName / TEXT("Engine/Binaries/Linux/UE4Editor");
+	FString FileName;
+	if (ExplicitFileName.IsEmpty())
+	{
+		FileName = RootDirName / TEXT("Engine/Binaries/Linux/UE4Editor");
+	}
+	else
+	{
+		FileName = ExplicitFileName;
+	}
 
-	FProcHandle ProcessHandle = FPlatformProcess::CreateProc(*CommandLine, *Arguments, true, false, false, nullptr, 0, nullptr, nullptr);
+	FProcHandle ProcessHandle = FPlatformProcess::CreateProc(*FileName, *Arguments, true, false, false, nullptr, 0, nullptr, nullptr);
 
 	if (ProcessHandle.IsValid())
 	{
diff --git a/Engine/Source/Programs/UnrealVersionSelector/Private/Linux/LinuxPlatformInstallation.h b/Engine/Source/Programs/UnrealVersionSelector/Private/Linux/LinuxPlatformInstallation.h
--- a/Engine/Source/Programs/UnrealVersionSelector/Private/Linux/LinuxPlatformInstallation.h
+++ b/Engine/Source/Programs/UnrealVersionSelector/Private/Linux/LinuxPlatformInstallation.h
@@ -8,7 +8,7 @@
 struct FLinuxPlatformInstallation : FGenericPlatformInstallation
 {
 	// Launches the editor application
-	static bool LaunchEditor(const FString &RootDirName, const FString &Arguments);
+	static bool LaunchEditor(const FString &RootDirName, const FString& ExplicitFileName, const FString &Arguments);
 
 	// Select an engine installation
 	static bool SelectEngineInstallation(FString &Identifier);
diff --git a/Engine/Source/Programs/UnrealVersionSelector/Private/UnrealVersionSelector.cpp b/Engine/Source/Programs/UnrealVersionSelector/Private/UnrealVersionSelector.cpp
--- a/Engine/Source/Programs/UnrealVersionSelector/Private/UnrealVersionSelector.cpp
+++ b/Engine/Source/Programs/UnrealVersionSelector/Private/UnrealVersionSelector.cpp
@@ -4,6 +4,7 @@
 #include "RequiredProgramMainCPPInclude.h"
 #include "DesktopPlatformModule.h"
 #include "PlatformInstallation.h"
+#include "Json/Public/Serialization/JsonSerializer.h"
 
 IMPLEMENT_APPLICATION(UnrealVersionSelector, "UnrealVersionSelector")
 
@@ -183,7 +184,7 @@ bool LaunchEditor()
 	FDesktopPlatformModule::Get()->GetEngineRootDirFromIdentifier(Identifier, RootDir);
 
 	// Launch the editor
-	if (!FPlatformInstallation::LaunchEditor(RootDir, FString()))
+	if (!FPlatformInstallation::LaunchEditor(RootDir, FString(), FString()))
 	{
 		FPlatformMisc::MessageBoxExt(EAppMsgType::Ok, TEXT("Failed to launch editor"), TEXT("Error"));
 		return false;
@@ -192,6 +193,91 @@ bool LaunchEditor()
 	return true;
 }
 
+bool ReadLaunchPathFromTargetFile(const FString& TargetFileName, FString& OutLaunchPath)
+{
+	// Read the file to a string
+	FString FileContents;
+	if (!FFileHelper::LoadFileToString(FileContents, *TargetFileName))
+	{
+		return false;
+	}
+
+	// Deserialize a JSON object from the string
+	TSharedPtr<FJsonObject> Object;
+	TSharedRef<TJsonReader<>> Reader = TJsonReaderFactory<>::Create(FileContents);
+	if (!FJsonSerializer::Deserialize(Reader, Object) || !Object.IsValid())
+	{
+		return false;
+	}
+
+	// Check it's an editor target
+	FString TargetType;
+	if(!Object->TryGetStringField(TEXT("TargetType"), TargetType) || TargetType != TEXT("Editor"))
+	{
+		return false;
+	}
+
+	// Check it's development configuration
+	FString Configuration;
+	if (!Object->TryGetStringField(TEXT("Configuration"), Configuration) || Configuration != TEXT("Development"))
+	{
+		return false;
+	}
+
+	// Get the launch path
+	OutLaunchPath.Empty();
+	Object->TryGetStringField(TEXT("Launch"), OutLaunchPath);
+	return true;
+}
+
+bool TryGetEditorFileName(const FString& EngineDir, const FString& ProjectFileName, FString& OutEditorFileName)
+{
+	IFileManager& FileManager = IFileManager::Get();
+
+	FString ProjectDir = FPaths::GetPath(ProjectFileName);
+	FString BinariesDir = ProjectDir / TEXT("Binaries") / FPlatformProcess::GetBinariesSubdirectory();
+	if (FileManager.DirectoryExists(*BinariesDir))
+	{
+		class FTargetFileVisitor : public IPlatformFile::FDirectoryStatVisitor
+		{
+		public:
+			TArray<TPair<FString, FDateTime>> Files;
+
+			virtual bool Visit(const TCHAR* FilenameOrDirectory, const FFileStatData& StatData)
+			{
+				static const TCHAR Extension[] = TEXT(".target");
+				static const int ExtensionLen = ARRAY_COUNT(Extension) - 1;
+
+				int Length = FCString::Strlen(FilenameOrDirectory);
+				if(Length >= ExtensionLen && FCString::Stricmp(FilenameOrDirectory + Length - ExtensionLen, Extension) == 0)
+				{
+					Files.Add(TPair<FString, FDateTime>(FilenameOrDirectory, StatData.ModificationTime));
+				}
+
+				return true;
+			}
+		};
+
+		FTargetFileVisitor Visitor;
+		FileManager.IterateDirectoryStat(*BinariesDir, Visitor);
+
+		Visitor.Files.Sort([](const TPair<FString, FDateTime>& A, const TPair<FString, FDateTime>& B){ return A.Value > B.Value; });
+
+		for(const TPair<FString, FDateTime>& Pair : Visitor.Files)
+		{
+			FString LaunchPath;
+			if(ReadLaunchPathFromTargetFile(Pair.Key, LaunchPath))
+			{
+				OutEditorFileName = MoveTemp(LaunchPath);
+				OutEditorFileName.ReplaceInline(TEXT("$(EngineDir)"), *EngineDir);
+				OutEditorFileName.ReplaceInline(TEXT("$(ProjectDir)"), *ProjectDir);
+				return true;
+			}
+		}
+	}
+	return false;
+}
+
 bool LaunchEditor(const FString& ProjectFileName, const FString& Arguments)
 {
 	// Get the engine root directory
@@ -201,8 +287,12 @@ bool LaunchEditor(const FString& ProjectFileName, const FString& Arguments)
 		return false;
 	}
 
+	// Figure out the path to the editor executable. This may be empty for older .target files; the platform layer should use the default path if necessary.
+	FString EditorFileName;
+	TryGetEditorFileName(RootDir / TEXT("Engine"), ProjectFileName, EditorFileName);
+
 	// Launch the editor
-	if (!FPlatformInstallation::LaunchEditor(RootDir, FString::Printf(TEXT("\"%s\" %s"), *ProjectFileName, *Arguments)))
+	if (!FPlatformInstallation::LaunchEditor(RootDir, EditorFileName, FString::Printf(TEXT("\"%s\" %s"), *ProjectFileName, *Arguments)))
 	{
 		FPlatformMisc::MessageBoxExt(EAppMsgType::Ok, TEXT("Failed to launch editor"), TEXT("Error"));
 		return false;
diff --git a/Engine/Source/Programs/UnrealVersionSelector/Private/Windows/WindowsPlatformInstallation.cpp b/Engine/Source/Programs/UnrealVersionSelector/Private/Windows/WindowsPlatformInstallation.cpp
--- a/Engine/Source/Programs/UnrealVersionSelector/Private/Windows/WindowsPlatformInstallation.cpp
+++ b/Engine/Source/Programs/UnrealVersionSelector/Private/Windows/WindowsPlatformInstallation.cpp
@@ -213,9 +213,19 @@ private:
 	}
 };
 
-bool FWindowsPlatformInstallation::LaunchEditor(const FString &RootDirName, const FString &Arguments)
+bool FWindowsPlatformInstallation::LaunchEditor(const FString &RootDirName, const FString& ExplicitFileName, const FString &Arguments)
 {
-	FString CommandLine = FString::Printf(TEXT("\"%s\" %s"), *(RootDirName / TEXT("Engine/Binaries/Win64/UE4Editor.exe")), *Arguments);
+	FString FileName;
+	if (ExplicitFileName.IsEmpty())
+	{
+		FileName = RootDirName / TEXT("Engine/Binaries/Win64/UE4Editor.exe");
+	}
+	else
+	{
+		FileName = ExplicitFileName;
+	}
+
+	FString CommandLine = FString::Printf(TEXT("\"%s\" %s"), *FileName, *Arguments);
 
 	STARTUPINFO StartupInfo;
 	ZeroMemory(&StartupInfo, sizeof(StartupInfo));
diff --git a/Engine/Source/Programs/UnrealVersionSelector/Private/Windows/WindowsPlatformInstallation.h b/Engine/Source/Programs/UnrealVersionSelector/Private/Windows/WindowsPlatformInstallation.h
--- a/Engine/Source/Programs/UnrealVersionSelector/Private/Windows/WindowsPlatformInstallation.h
+++ b/Engine/Source/Programs/UnrealVersionSelector/Private/Windows/WindowsPlatformInstallation.h
@@ -8,7 +8,7 @@
 struct FWindowsPlatformInstallation : FGenericPlatformInstallation
 {
 	// Launches the editor application
-	static bool LaunchEditor(const FString &RootDirName, const FString &Arguments);
+	static bool LaunchEditor(const FString &RootDirName, const FString& ExplicitFileName, const FString &Arguments);
 
 	// Select an engine installation
 	static bool SelectEngineInstallation(FString &Identifier);
diff --git a/Engine/Source/Programs/UnrealVersionSelector/UnrealVersionSelector.Build.cs b/Engine/Source/Programs/UnrealVersionSelector/UnrealVersionSelector.Build.cs
--- a/Engine/Source/Programs/UnrealVersionSelector/UnrealVersionSelector.Build.cs
+++ b/Engine/Source/Programs/UnrealVersionSelector/UnrealVersionSelector.Build.cs
@@ -18,5 +18,7 @@ public class UnrealVersionSelector : ModuleRules
 		{
 			PrivateDependencyModuleNames.AddRange(new string[] { "Core", "Projects", "DesktopPlatform", "ApplicationCore" });
 		}
-	}
+
+        PrivateDependencyModuleNames.Add("Json");
+    }
 }
