Fix manifest error linking UnrealVersionSelector.

diff --git a/Engine/Source/Programs/UnrealVersionSelector/Private/Windows/Resources/Resource.rc b/Engine/Source/Programs/UnrealVersionSelector/Private/Windows/Resources/Resource.rc
--- a/Engine/Source/Programs/UnrealVersionSelector/Private/Windows/Resources/Resource.rc
+++ b/Engine/Source/Programs/UnrealVersionSelector/Private/Windows/Resources/Resource.rc
@@ -8,7 +8,7 @@
 // Generated from the TEXTINCLUDE 2 resource.
 //
 #include "winres.h"
-#include "VersionResource.inl"
+#include "Resource.rc2"
 
 /////////////////////////////////////////////////////////////////////////////
 #undef APSTUDIO_READONLY_SYMBOLS
@@ -33,7 +33,7 @@ END
 2 TEXTINCLUDE 
 BEGIN
     "#include ""winres.h""\r\n"
-    "#include ""VersionResource.inl""\r\n"
+    "#include ""Resource.rc2""\r\n"
     "\0"
 END
 
diff --git a/Engine/Source/Programs/UnrealVersionSelector/Private/Windows/Resources/Resource.rc2 b/Engine/Source/Programs/UnrealVersionSelector/Private/Windows/Resources/Resource.rc2
new file mode 100644
--- /dev/null
+++ b/Engine/Source/Programs/UnrealVersionSelector/Private/Windows/Resources/Resource.rc2
@@ -0,0 +1,52 @@
+// Copyright Epic Games, Inc. All Rights Reserved.
+
+#ifndef APSTUDIO_INVOKED
+
+#include <windows.h>
+#include "../../../../../Runtime/Launch/Resources/Windows/resource.h"
+#include "../../../../../Runtime/Launch/Resources/Version.h"
+
+
+/////////////////////////////////////////////////////////////////////////////
+//
+// Version
+//
+
+VS_VERSION_INFO VERSIONINFO
+ FILEVERSION ENGINE_MAJOR_VERSION,ENGINE_MINOR_VERSION,ENGINE_PATCH_VERSION,0
+ PRODUCTVERSION ENGINE_MAJOR_VERSION,ENGINE_MINOR_VERSION,ENGINE_PATCH_VERSION,0
+ FILEFLAGSMASK 0x17L
+#ifdef _DEBUG
+ FILEFLAGS 0x1L
+#else
+ FILEFLAGS 0x0L
+#endif
+ FILEOS 0x4L
+ FILETYPE 0x2L
+ FILESUBTYPE 0x0L
+BEGIN
+	BLOCK "StringFileInfo"
+	BEGIN
+		BLOCK "040904b0"
+		BEGIN
+			VALUE "CompanyName", EPIC_COMPANY_NAME
+			VALUE "LegalCopyright", EPIC_COPYRIGHT_STRING
+			VALUE "ProductName", EPIC_PRODUCT_NAME
+			VALUE "ProductVersion", ENGINE_VERSION_STRING
+			VALUE "FileDescription", EPIC_PRODUCT_NAME
+			VALUE "InternalName", EPIC_PRODUCT_IDENTIFIER
+			VALUE "OriginalFilename", "VersionSelector.exe"
+			VALUE "ShellIntegrationVersion", "2"
+		END
+	END
+	BLOCK "VarFileInfo"
+	BEGIN
+		VALUE "Translation", 0x409, 1200
+	END
+END
+
+#endif
+
+// Icon with lowest ID value placed first to ensure application icon
+// remains consistent on all systems.
+IDICON_UE4Game			ICON                    "Icon.ico"
diff --git a/Engine/Source/Programs/UnrealVersionSelector/Private/Windows/Resources/VersionResource.inl b/Engine/Source/Programs/UnrealVersionSelector/Private/Windows/Resources/VersionResource.inl
deleted file mode 100644
