UnrealVersionSelector: Fix include path from parent directory

#rnx

diff --git a/Engine/Source/Programs/UnrealVersionSelector/Private/UnrealVersionSelector.cpp b/Engine/Source/Programs/UnrealVersionSelector/Private/UnrealVersionSelector.cpp
--- a/Engine/Source/Programs/UnrealVersionSelector/Private/UnrealVersionSelector.cpp
+++ b/Engine/Source/Programs/UnrealVersionSelector/Private/UnrealVersionSelector.cpp
@@ -4,7 +4,7 @@
 #include "RequiredProgramMainCPPInclude.h"
 #include "DesktopPlatformModule.h"
 #include "PlatformInstallation.h"
-#include "Json/Public/Serialization/JsonSerializer.h"
+#include "Serialization/JsonSerializer.h"
 
 IMPLEMENT_APPLICATION(UnrealVersionSelector, "UnrealVersionSelector")
 
