Fix messagebox not working in version selector due to missing resources in exe. Similar to CL 23287701.

diff --git a/Engine/Source/Programs/LiveLinkHub/Resources/Windows/LiveLinkHub.rc b/Engine/Source/Programs/LiveLinkHub/Resources/Windows/LiveLinkHub.rc
--- a/Engine/Source/Programs/LiveLinkHub/Resources/Windows/LiveLinkHub.rc
+++ b/Engine/Source/Programs/LiveLinkHub/Resources/Windows/LiveLinkHub.rc
@@ -7,7 +7,7 @@
 //
 #include <windows.h>
 #include "Runtime/Launch/Resources/Version.h"
-#include "Runtime/Launch/Resources/Windows/resource.h"
+#include "Runtime/Launch/Resources/Windows/PCLaunch.rc"
 
 #define LIVELINK_HUB_VERSION_MAJOR 1
 #define LIVELINK_HUB_VERSION_MINOR 0
diff --git a/Engine/Source/Programs/UnrealVersionSelector/Private/Windows/Resources/Resource.h b/Engine/Source/Programs/UnrealVersionSelector/Private/Windows/Resources/Resource.h
--- a/Engine/Source/Programs/UnrealVersionSelector/Private/Windows/Resources/Resource.h
+++ b/Engine/Source/Programs/UnrealVersionSelector/Private/Windows/Resources/Resource.h
@@ -4,21 +4,21 @@
 // Microsoft Visual C++ generated include file.
 // Used by Resource.rc
 //
-#define IDD_SELECTBUILD                 101
-#define IDD_ERRORDIALOG                 105
-#define IDC_BROWSE                      1001
-#define IDC_BUILDLIST                   1002
-#define IDC_ERRORLOGTEXT                1005
-#define IDC_ERRORMESSAGE                1006
+#define IDD_SELECTBUILD                 201
+#define IDD_ERRORDIALOG                 205
+#define IDC_BROWSE                      2001
+#define IDC_BUILDLIST                   2002
+#define IDC_ERRORLOGTEXT                2005
+#define IDC_ERRORMESSAGE                2006
 
 // Next default values for new objects
 // 
 #ifdef APSTUDIO_INVOKED
 #ifndef APSTUDIO_READONLY_SYMBOLS
-#define _APS_NEXT_RESOURCE_VALUE        106
+#define _APS_NEXT_RESOURCE_VALUE        206
 #define _APS_NEXT_COMMAND_VALUE         40001
-#define _APS_NEXT_CONTROL_VALUE         1007
-#define _APS_NEXT_SYMED_VALUE           101
+#define _APS_NEXT_CONTROL_VALUE         2007
+#define _APS_NEXT_SYMED_VALUE           201
 #endif
 #endif
 
diff --git a/Engine/Source/Programs/UnrealVersionSelector/Private/Windows/Resources/Resource.rc2 b/Engine/Source/Programs/UnrealVersionSelector/Private/Windows/Resources/Resource.rc2
--- a/Engine/Source/Programs/UnrealVersionSelector/Private/Windows/Resources/Resource.rc2
+++ b/Engine/Source/Programs/UnrealVersionSelector/Private/Windows/Resources/Resource.rc2
@@ -3,7 +3,7 @@
 #ifndef APSTUDIO_INVOKED
 
 #include <windows.h>
-#include "../../../../../Runtime/Launch/Resources/Windows/resource.h"
+#include "../../../../../Runtime/Launch/Resources/Windows/PCLaunch.rc"
 #include "../../../../../Runtime/Launch/Resources/Version.h"
 
 
