Add #include for FStringOutputDevice now that it is separated into its own header.

Also, addressed some preexisting warnings issued by the MSVC static analyzer which prevented SubmitTool from validating.

diff --git a/Engine/Source/Programs/UnrealVersionSelector/Private/UnrealVersionSelector.cpp b/Engine/Source/Programs/UnrealVersionSelector/Private/UnrealVersionSelector.cpp
--- a/Engine/Source/Programs/UnrealVersionSelector/Private/UnrealVersionSelector.cpp
+++ b/Engine/Source/Programs/UnrealVersionSelector/Private/UnrealVersionSelector.cpp
@@ -4,6 +4,7 @@
 #include "RequiredProgramMainCPPInclude.h"
 #include "DesktopPlatformModule.h"
 #include "PlatformInstallation.h"
+#include "Misc/StringOutputDevice.h"
 #include "Serialization/JsonSerializer.h"
 
 IMPLEMENT_APPLICATION(UnrealVersionSelector, "UnrealVersionSelector")
@@ -406,7 +407,7 @@ int Main(const TArray<FString>& Arguments)
 	#include "Windows/AllowWindowsPlatformTypes.h"
 	#include <Shellapi.h>
 
-	int WINAPI WinMain(HINSTANCE hCurrInstance, HINSTANCE hPrevInstance, LPSTR lpCmdLine, int ShowCmd)
+	int WINAPI WinMain(_In_ HINSTANCE /*hInstance*/, _In_opt_ HINSTANCE /*hPrevInstance*/, _In_ LPSTR /*lpCmdLine*/, _In_ int /*nShowCmd*/)
 	{
 		int ArgC;
 		LPWSTR* ArgV = ::CommandLineToArgvW(GetCommandLineW(), &ArgC);
diff --git a/Engine/Source/Runtime/Engine/Private/Tests/TextImportTest.cpp b/Engine/Source/Runtime/Engine/Private/Tests/TextImportTest.cpp
--- a/Engine/Source/Runtime/Engine/Private/Tests/TextImportTest.cpp
+++ b/Engine/Source/Runtime/Engine/Private/Tests/TextImportTest.cpp
@@ -2,6 +2,7 @@
 
 #include "Tests/TextImportTest.h"
 #include "Tests/TestHarnessAdapter.h"
+#include "Misc/StringOutputDevice.h"
 #include "UObject/UnrealType.h"
 
 #if WITH_TESTS
