UnrealVersionSelector: Clarify error popup message when engine version change fails.

#rnx

#Synced-CL 50656828

diff --git a/Engine/Source/Programs/UnrealVersionSelector/Private/UnrealVersionSelector.cpp b/Engine/Source/Programs/UnrealVersionSelector/Private/UnrealVersionSelector.cpp
--- a/Engine/Source/Programs/UnrealVersionSelector/Private/UnrealVersionSelector.cpp
+++ b/Engine/Source/Programs/UnrealVersionSelector/Private/UnrealVersionSelector.cpp
@@ -100,7 +100,7 @@ bool SwitchVersion(const FString& ProjectFileName)
 	// Update the project file
 	if (!FDesktopPlatformModule::Get()->SetEngineIdentifierForProject(ProjectFileName, Identifier))
 	{
-		FPlatformMisc::MessageBoxExt(EAppMsgType::Ok, TEXT("Couldn't set association for project. Check the file is writeable."), TEXT("Error"));
+		FPlatformMisc::MessageBoxExt(EAppMsgType::Ok, TEXT("Couldn't set association for project. Check if the uproject file contains a valid JSON and if the file is writeable."), TEXT("Error"));
 		return false;
 	}
 
